CrashDebugHelper - Removed obsolete code used to retrieve the label string from the database, implemented basic mechanism for caching the PDB files

diff --git a/Engine/Config/BaseEngine.ini b/Engine/Config/BaseEngine.ini
--- a/Engine/Config/BaseEngine.ini
+++ b/Engine/Config/BaseEngine.ini
@@ -2008,9 +2008,14 @@ AlwaysRecordingStats=true
 [Engine.CrashDebugHelper]
 ; ini files do not like '//' - they treat it as a comment
 LocalSymbolStore=d:/Symbols/UE4
-BuilderDatabase=db-09
-DatabaseCatalog=Builder
 SourceControlBuildLabelPattern=*-CL-%CHANGELISTNUMBER%
+; devweb-02
+PDBCachePath=F:/CrashReportPDBCache/
+;PDBCachePath=U:/CrashReport/PDBCache/
+DaysToDeleteUnusedFilesFromPDBCache=14
+PDBCacheSizeGB=128
+MinDiskFreeSpaceGB=64
+bUsePDBCache=true
 
 [Engine.ErrorHandling]
 bPromptForRemoteDebugging=false
diff --git a/Engine/Source/Developer/CrashDebugHelper/Private/CrashDebugHelper.cpp b/Engine/Source/Developer/CrashDebugHelper/Private/CrashDebugHelper.cpp
--- a/Engine/Source/Developer/CrashDebugHelper/Private/CrashDebugHelper.cpp
+++ b/Engine/Source/Developer/CrashDebugHelper/Private/CrashDebugHelper.cpp
@@ -63,19 +63,7 @@ bool ICrashDebugHelper::Init()
 		bInitialized = false;
 	}
 
-	// Look up the builder database name - fail if not found
-	if (GConfig->GetString(TEXT("Engine.CrashDebugHelper"), TEXT("BuilderDatabase"), DatabaseName, GEngineIni) == false)
-	{
-		UE_LOG(LogCrashDebugHelper, Warning, TEXT("Failed to get BuilderDatabase from ini file... crash handling disabled"));
-		bInitialized = false;
-	}
-
-	// Look up the builder database catalog - fail if not found
-	if (GConfig->GetString(TEXT("Engine.CrashDebugHelper"), TEXT("DatabaseCatalog"), DatabaseCatalog, GEngineIni) == false)
-	{
-		UE_LOG(LogCrashDebugHelper, Warning, TEXT("Failed to get builder DatabaseCatalog from ini file... crash handling disabled"));
-		bInitialized = false;
-	}
+	PDBCache.Init();
 
 	return bInitialized;
 }
@@ -147,74 +135,101 @@ bool ICrashDebugHelper::SyncModules( FCrashInfo* CrashInfo )
 	TArray< TSharedRef<ISourceControlLabel> > Labels = ISourceControlModule::Get().GetProvider().GetLabels( CrashInfo->LabelName );
 	if(Labels.Num() > 0)
 	{
-		// Sync every module from every label. If the same modules appear in every label, this will fail.
-		for (auto LabelIt = Labels.CreateConstIterator(); LabelIt; ++LabelIt)
+		const bool bContains = PDBCache.UsePDBCache() && PDBCache.ContainsPDBCacheEntry( CrashInfo->LabelName );
+		if( bContains )
+		{
+			UE_LOG( LogCrashDebugHelper, Warning, TEXT( "Label %s found in the PDB Cache, using it" ), *CrashInfo->LabelName );
+			CrashInfo->PDBCacheEntry = PDBCache.FindPDBCacheEntry( CrashInfo->LabelName );
+			PDBCache.TouchPDBCacheEntry( CrashInfo->LabelName );
+			return true;
+		}
+		else
 		{
-			TSharedRef<ISourceControlLabel> Label = *LabelIt;
+			TArray<FString> SyncedFiles;
 
-			//@TODO: MAC: Excluding labels for Mac since we are only syncing windows binaries here...
-			if (Label->GetName().Contains(TEXT("Mac")))
+			// Sync every module from every label. If the same modules appear in every label, this will fail.
+			for( auto LabelIt = Labels.CreateConstIterator(); LabelIt; ++LabelIt )
 			{
-				UE_LOG(LogCrashDebugHelper, Log, TEXT(" Skipping Mac label '%s' when syncing modules."), *Label->GetName());
-			}
-			else
-			{
-				UE_LOG(LogCrashDebugHelper, Log, TEXT(" Syncing modules with label '%s'."), *Label->GetName());
-
-		        TArray<FString> FilesToSync;
-		        for( int32 ModuleNameIndex = 0; ModuleNameIndex < CrashInfo->ModuleNames.Num(); ModuleNameIndex++ )
-		        {
-			        FString DepotPath = FString::Printf( TEXT( "%s/%s" ), *DepotName, *CrashInfo->ModuleNames[ModuleNameIndex] );
-			        // Match all decorated versions of the module. We may need them if the file was renamed to remove the "-Platform-Configuration" decoration.
-			        DepotPath = DepotPath.Replace(TEXT(".dll"), TEXT("*.dll")).Replace(TEXT(".exe"), TEXT("*.exe"));
-        
-			        {
-				        TSharedRef<ISourceControlLabel> Label = Labels[0];
-				        if( Label->Sync(DepotPath) )
-				        {
-					        UE_LOG( LogCrashDebugHelper, Warning, TEXT( " ... synced binary '%s'."), *DepotPath );
-				        }
-        
-				        FString PDBName = DepotPath.Replace( TEXT( ".dll" ), TEXT( ".pdb" ) ).Replace( TEXT( ".exe" ), TEXT( ".pdb" ) );
-				        if( Label->Sync(PDBName) )
-				        {
-					        UE_LOG( LogCrashDebugHelper, Warning, TEXT( " ... synced symbol '%s'."), *PDBName );
-				        }
-        
-				        //@TODO: ROCKETHACK: Adding additional Installed and Symbol paths - revisit when builds are made by the builder...
-				        DepotPath = FString::Printf( TEXT( "%s/Rocket/Installed/Windows/%s" ), *DepotName, *CrashInfo->ModuleNames[ModuleNameIndex] );
-				        // Match all decorated versions of the module. We may need them if the file was renamed to remove the "-Platform-Configuration" decoration.
-				        DepotPath = DepotPath.Replace(TEXT(".dll"), TEXT("*.dll")).Replace(TEXT(".exe"), TEXT("*.exe"));
-				        if( Label->Sync(DepotPath) )
-				        {
-					        UE_LOG( LogCrashDebugHelper, Warning, TEXT( " ... synced binary '%s'."), *DepotPath );
-				        }
-        
-				        DepotPath = FString::Printf( TEXT( "%s/Rocket/Symbols/%s" ), *DepotName, *CrashInfo->ModuleNames[ModuleNameIndex] );
-				        PDBName = DepotPath.Replace( TEXT( ".dll" ), TEXT( "*.pdb" ) ).Replace( TEXT( ".exe" ), TEXT( "*.pdb" ) );
-				        if( Label->Sync(PDBName) )
-				        {
-					        UE_LOG( LogCrashDebugHelper, Warning, TEXT( " ... synced symbol '%s'."), *PDBName );
-				        }
-        
-				        DepotPath = FString::Printf( TEXT( "%s/Rocket/LauncherInstalled/Windows/Launcher/%s" ), *DepotName, *CrashInfo->ModuleNames[ModuleNameIndex] );
-				        // Match all decorated versions of the module. We may need them if the file was renamed to remove the "-Platform-Configuration" decoration.
-				        DepotPath = DepotPath.Replace(TEXT(".dll"), TEXT("*.dll")).Replace(TEXT(".exe"), TEXT("*.exe"));
-				        if( Label->Sync(DepotPath) )
-				        {
-					        UE_LOG( LogCrashDebugHelper, Warning, TEXT( " ... synced binary '%s'."), *DepotPath );
-				        }
-        
-				        DepotPath = FString::Printf( TEXT( "%s/Rocket/LauncherSymbols/Windows/Launcher/%s" ), *DepotName, *CrashInfo->ModuleNames[ModuleNameIndex] );
-				        PDBName = DepotPath.Replace( TEXT( ".dll" ), TEXT( "*.pdb" ) ).Replace( TEXT( ".exe" ), TEXT( "*.pdb" ) );
-				        if( Label->Sync(PDBName) )
-				        {
-					        UE_LOG( LogCrashDebugHelper, Warning, TEXT( " ... synced symbol '%s'."), *PDBName );
-				        }
+				TSharedRef<ISourceControlLabel> Label = *LabelIt;
+
+				//@TODO: MAC: Excluding labels for Mac since we are only syncing windows binaries here...
+				if( Label->GetName().Contains( TEXT( "Mac" ) ) )
+				{
+					UE_LOG( LogCrashDebugHelper, Log, TEXT( " Skipping Mac label '%s' when syncing modules." ), *Label->GetName() );
+				}
+				else
+				{
+					UE_LOG( LogCrashDebugHelper, Log, TEXT( " Syncing modules with label '%s'." ), *Label->GetName() );
+
+					TArray<FString> FilesToSync;
+					for( int32 ModuleNameIndex = 0; ModuleNameIndex < CrashInfo->ModuleNames.Num(); ModuleNameIndex++ )
+					{
+						// Sometimes P4 may hang due to spamming the server with the commands, if so try uncommenting the line below.
+						//FPlatformProcess::Sleep( 0.1f );
+						FString DepotPath = FString::Printf( TEXT( "%s/%s" ), *DepotName, *CrashInfo->ModuleNames[ModuleNameIndex] );
+						// Match all decorated versions of the module. We may need them if the file was renamed to remove the "-Platform-Configuration" decoration.
+						DepotPath = DepotPath.Replace( TEXT( ".dll" ), TEXT( "*.dll" ) ).Replace( TEXT( ".exe" ), TEXT( "*.exe" ) );
+
+						{
+							TSharedRef<ISourceControlLabel> Label = Labels[0];
+							if( Label->Sync( DepotPath ) )
+							{
+								UE_LOG( LogCrashDebugHelper, Warning, TEXT( " ... synced binary '%s'." ), *DepotPath );
+							}
+							SyncedFiles.Add( DepotPath );
+
+				        	FString PDBName = DepotPath.Replace( TEXT( ".dll" ), TEXT( ".pdb" ) ).Replace( TEXT( ".exe" ), TEXT( ".pdb" ) );
+							if( Label->Sync( PDBName ) )
+							{
+								UE_LOG( LogCrashDebugHelper, Warning, TEXT( " ... synced symbol '%s'." ), *PDBName );
+							}
+							SyncedFiles.Add( PDBName );
+
+							//@TODO: ROCKETHACK: Adding additional Installed and Symbol paths - revisit when builds are made by the builder...
+							DepotPath = FString::Printf( TEXT( "%s/Rocket/Installed/Windows/%s" ), *DepotName, *CrashInfo->ModuleNames[ModuleNameIndex] );
+							// Match all decorated versions of the module. We may need them if the file was renamed to remove the "-Platform-Configuration" decoration.
+							DepotPath = DepotPath.Replace( TEXT( ".dll" ), TEXT( "*.dll" ) ).Replace( TEXT( ".exe" ), TEXT( "*.exe" ) );
+							if( Label->Sync( DepotPath ) )
+							{
+								UE_LOG( LogCrashDebugHelper, Warning, TEXT( " ... synced binary '%s'." ), *DepotPath );
+							}
+							SyncedFiles.Add( DepotPath );
+
+							DepotPath = FString::Printf( TEXT( "%s/Rocket/Symbols/%s" ), *DepotName, *CrashInfo->ModuleNames[ModuleNameIndex] );
+							PDBName = DepotPath.Replace( TEXT( ".dll" ), TEXT( "*.pdb" ) ).Replace( TEXT( ".exe" ), TEXT( "*.pdb" ) );
+							if( Label->Sync( PDBName ) )
+							{
+								UE_LOG( LogCrashDebugHelper, Warning, TEXT( " ... synced symbol '%s'." ), *PDBName );		
+							}
+							SyncedFiles.Add( PDBName );
+
+							DepotPath = FString::Printf( TEXT( "%s/Rocket/LauncherInstalled/Windows/Launcher/%s" ), *DepotName, *CrashInfo->ModuleNames[ModuleNameIndex] );
+							// Match all decorated versions of the module. We may need them if the file was renamed to remove the "-Platform-Configuration" decoration.
+							DepotPath = DepotPath.Replace( TEXT( ".dll" ), TEXT( "*.dll" ) ).Replace( TEXT( ".exe" ), TEXT( "*.exe" ) );
+							if( Label->Sync( DepotPath ) )
+							{
+								UE_LOG( LogCrashDebugHelper, Warning, TEXT( " ... synced binary '%s'." ), *DepotPath );
+							}
+							SyncedFiles.Add( DepotPath );
+
+							DepotPath = FString::Printf( TEXT( "%s/Rocket/LauncherSymbols/Windows/Launcher/%s" ), *DepotName, *CrashInfo->ModuleNames[ModuleNameIndex] );
+							PDBName = DepotPath.Replace( TEXT( ".dll" ), TEXT( "*.pdb" ) ).Replace( TEXT( ".exe" ), TEXT( "*.pdb" ) );
+							if( Label->Sync( PDBName ) )
+							{
+								UE_LOG( LogCrashDebugHelper, Warning, TEXT( " ... synced symbol '%s'." ), *PDBName );
+							}
+							SyncedFiles.Add( PDBName );
+						}
 					}
 				}
 			}
-		}
+
+			if( PDBCache.UsePDBCache() )
+			{
+				// Initialize and add a new PDB Cache entry to the database.
+				CrashInfo->PDBCacheEntry = PDBCache.CreateAndAddPDBCacheEntry( CrashInfo->LabelName, DepotName, SyncedFiles );
+			}
+		}		
 	}
 	else
 	{
@@ -639,105 +654,321 @@ FString ICrashDebugHelper::RetrieveBuildLabel(int32 InEngineVersion, int32 InCha
 {
 	FString FoundLabelString;
 
-	if ((DatabaseName.Len() == 0) || (DatabaseCatalog.Len() == 0))
+	if ((InEngineVersion < 0) && (InChangelistNumber < 0))
 	{
-		UE_LOG(LogCrashDebugHelper, Warning, TEXT("RetrieveBuildLabel: Invalid databasename and/or databasecatalog set."));
+		UE_LOG(LogCrashDebugHelper, Warning, TEXT("RetrieveBuildLabel: Invalid parameters."));
 		return FoundLabelString;
 	}
 
-	if ((InEngineVersion < 0) && (InChangelistNumber < 0))
+	// If we failed to find the label, try to find the label directly in source control by using the pattern supplied via ini
+	if ( FoundLabelString.IsEmpty() && InChangelistNumber >= 0 && !SourceControlBuildLabelPattern.IsEmpty() )
 	{
-		UE_LOG(LogCrashDebugHelper, Warning, TEXT("RetrieveBuildLabel: Invalid parameters."));
-		return FoundLabelString;
+		const FString ChangelistString = FString::Printf(TEXT("%d"), InChangelistNumber);
+		const FString TestLabel = SourceControlBuildLabelPattern.Replace(TEXT("%CHANGELISTNUMBER%"), *ChangelistString, ESearchCase::CaseSensitive );
+		TArray< TSharedRef<ISourceControlLabel> > Labels = ISourceControlModule::Get().GetProvider().GetLabels( TestLabel );
+		if ( Labels.Num() > 0 )
+		{
+			// If we found more than one label, warn about it and just use the first one
+			if ( Labels.Num() > 1 )
+			{
+				UE_LOG(LogCrashDebugHelper, Warning, TEXT("RetrieveBuildLabel: More than one build label found with pattern %s - Using label %s"), *TestLabel, *Labels[0]->GetName());
+			}
+
+			FoundLabelString = Labels[0]->GetName();
+			UE_LOG(LogCrashDebugHelper, Log, TEXT("RetrieveBuildLabel: Found label %s matching pattern %s in source control."), *FoundLabelString, *TestLabel);
+		}
 	}
 
-	// Open a database connection
-	// Create the connection object; needs to be deleted via "delete".
-	FDataBaseConnection* Connection = FDataBaseConnection::CreateObject();
-	check(Connection);
+	return FoundLabelString;
+}
 
-	// Create the connection string with Windows Authentication as the way to handle permissions/ login/ security.
-	FString ConnectionString = FString::Printf(
-		TEXT("Provider=sqloledb;Data Source=%s;Initial Catalog=%s;Trusted_Connection=Yes;"), *DatabaseName, *DatabaseCatalog);
+/*-----------------------------------------------------------------------------
+	PDB Cache implementation
+-----------------------------------------------------------------------------*/
 
-	// Try to open connection to DB - this is a synchronous operation.
-	if (Connection->Open(*ConnectionString, NULL, NULL))
-	{
-		UE_LOG(LogCrashDebugHelper, Warning, TEXT("RetrieveBuildLabel: Connection to %s.%s succeeded"), *DatabaseName, *DatabaseCatalog);
+const TCHAR* FPDBCache::PDBTimeStampFile = TEXT( "PDBTimeStamp.txt" );
 
-		// Setup the query command
-		FString LabelRequestCmd = TEXT("SELECT Label FROM Versioning WHERE ");
-		if (InEngineVersion >= 0)
+void FPDBCache::Init()
+{
+	// PDB Cache
+	// Default configuration
+	// PDBCachePath=U:/CrashReport/PDBCache
+	// DaysToDeleteUnusedFilesFromPDBCache=14
+	// PDBCacheSizeGB=128
+	// MinFreeSizeGB=64
+	// bUsePDBCache=true
+
+	// Look up whether we want to use the PDB cache.
+	GConfig->GetBool( TEXT( "Engine.CrashDebugHelper" ), TEXT( "bUsePDBCache" ), bUsePDBCache, GEngineIni );
+	UE_LOG( LogCrashDebugHelper, Warning, TEXT( "bUsePDBCache is %s" ), bUsePDBCache ? TEXT( "enabled" ) : TEXT( "disabled" ) );
+
+	// Get the rest of the PDB cache configuration.
+	if( bUsePDBCache )
+	{
+		if( !GConfig->GetString( TEXT( "Engine.CrashDebugHelper" ), TEXT( "PDBCachePath" ), PDBCachePath, GEngineIni ) )
 		{
-			LabelRequestCmd += FString::Printf(TEXT("EngineVersion = %d"), InEngineVersion);
+			UE_LOG( LogCrashDebugHelper, Warning, TEXT( "Failed to get PDBCachePath from ini file... PDB Cache disabled" ) );
+			bUsePDBCache = false;
 		}
-		else
+	}
+
+	if( bUsePDBCache )
+	{
+		if( !GConfig->GetInt( TEXT( "Engine.CrashDebugHelper" ), TEXT( "PDBCacheSizeGB" ), PDBCacheSizeGB, GEngineIni ) )
 		{
-			check(InChangelistNumber >= 0);
-			LabelRequestCmd += FString::Printf(TEXT("Changelist = %d"), InChangelistNumber);
+			UE_LOG( LogCrashDebugHelper, Warning, TEXT( "Failed to get PDBCachePath from ini file... Using default value" ) );
 		}
 
-		// Run the query
-		FDataBaseRecordSet* NewRecordSet = NULL;
-		if (Connection->Execute(*LabelRequestCmd, NewRecordSet) == false)
+		if( !GConfig->GetInt( TEXT( "Engine.CrashDebugHelper" ), TEXT( "MinDiskFreeSpaceGB" ), MinDiskFreeSpaceGB, GEngineIni ) )
 		{
-			UE_LOG(LogCrashDebugHelper, Warning, TEXT("RetrieveBuildLabel: Failed to exec SQLCommand for..."));
-			UE_LOG(LogCrashDebugHelper, Warning, TEXT("\t%s"), *LabelRequestCmd);
+			UE_LOG( LogCrashDebugHelper, Warning, TEXT( "Failed to get MinDiskFreeSpaceGB from ini file... Using default value" ) );
 		}
-		else
+
+		if( !GConfig->GetInt( TEXT( "Engine.CrashDebugHelper" ), TEXT( "DaysToDeleteUnusedFilesFromPDBCache" ), DaysToDeleteUnusedFilesFromPDBCache, GEngineIni ) )
 		{
-			if ((NewRecordSet == NULL) || (NewRecordSet->GetRecordCount() == 0))
+			UE_LOG( LogCrashDebugHelper, Warning, TEXT( "Failed to get DaysToDeleteUnusedFilesFromPDBCache from ini file... Using default value" ) );
+		}
+
+		InitializePDBCache();
+		CleanPDBCache( DaysToDeleteUnusedFilesFromPDBCache );
+
+		// Verify that we have enough space to enable the PDB Cache.
+		uint64 TotalNumberOfBytes = 0;
+		uint64 NumberOfFreeBytes = 0;
+		FPlatformMisc::GetDiskTotalAndFreeSpace( PDBCachePath, TotalNumberOfBytes, NumberOfFreeBytes );
+
+		const int32 TotalDiscSpaceGB = int32( TotalNumberOfBytes >> 30 );
+		const int32 DiskFreeSpaceGB = int32( NumberOfFreeBytes >> 30 );
+
+		if( DiskFreeSpaceGB < MinDiskFreeSpaceGB )
+		{
+			// There is not enough free space, calculate the current PDB cache usage and try removing the old data.
+			const int32 CurrentPDBCacheSizeGB = GetPDBCacheSizeGB();
+			const int32 DiskFreeSpaceAfterCleanGB = DiskFreeSpaceGB + CurrentPDBCacheSizeGB;
+			if( DiskFreeSpaceAfterCleanGB < MinDiskFreeSpaceGB )
 			{
-				if (InEngineVersion >= 0)
-				{
-					UE_LOG(LogCrashDebugHelper, Warning, TEXT("RetrieveBuildLabel: No results for engine version %d"), InEngineVersion);
-				}
-				else
-				{
-					UE_LOG(LogCrashDebugHelper, Warning, TEXT("RetrieveBuildLabel: No results for changelist %d"), InChangelistNumber);
-				}
+				UE_LOG( LogCrashDebugHelper, Error, TEXT( "There is not enough free space. PDB Cache disabled." ) );
+				UE_LOG( LogCrashDebugHelper, Error, TEXT( "Current disk free space is %i GBs." ), DiskFreeSpaceGB );
+				UE_LOG( LogCrashDebugHelper, Error, TEXT( "To enable the PDB Cache you need to free %i GB of space" ), MinDiskFreeSpaceGB - DiskFreeSpaceAfterCleanGB );
+				bUsePDBCache = false;
+				// Remove all data.
+				CleanPDBCache( 0 );
 			}
 			else
 			{
-				for (FDataBaseRecordSet::TIterator It(NewRecordSet); It; ++It)
-				{
-					FoundLabelString = It->GetString(TEXT("Label"));
-					FoundLabelString.TrimTrailing();
-				}
+				// Clean the PDB cache until we get enough free space.
+				CleanPDBCache( DaysToDeleteUnusedFilesFromPDBCache, MinDiskFreeSpaceGB - DiskFreeSpaceGB );
 			}
 		}
+	}
 
-		Connection->Close();
+	if( bUsePDBCache )
+	{
+		UE_LOG( LogCrashDebugHelper, Log, TEXT( "PDBCachePath=%s" ), *PDBCachePath );
+		UE_LOG( LogCrashDebugHelper, Log, TEXT( "PDBCacheSizeGB=%i" ), PDBCacheSizeGB );
+		UE_LOG( LogCrashDebugHelper, Log, TEXT( "MinDiskFreeSpaceGB=%i" ), MinDiskFreeSpaceGB );
+		UE_LOG( LogCrashDebugHelper, Log, TEXT( "DaysToDeleteUnusedFilesFromPDBCache=%i" ), DaysToDeleteUnusedFilesFromPDBCache );
 	}
-	else
+}
+
+
+void FPDBCache::InitializePDBCache()
+{
+	const double StartTime = FPlatformTime::Seconds();
+
+	TArray<FString> PDBCacheEntryDirectories;
+	IFileManager::Get().FindFiles( PDBCacheEntryDirectories, *PDBCachePath, false, true );
+
+	for( const auto& DirPath : PDBCacheEntryDirectories )
 	{
-		// Connection failed :(
-		UE_LOG(LogCrashDebugHelper, Warning, TEXT("RetrieveBuildLabel: Connection to %s.%s failed"), *DatabaseName, *DatabaseCatalog);
+		FPDBCacheEntryRef Entry = ReadPDBCacheEntry( DirPath );
+		PDBCacheEntries.Add( DirPath, Entry );
 	}
 
-	// delete
-	delete Connection;
-	Connection = NULL;
+	SortPDBCache();
 
-	// If we failed to find the label, try to find the label directly in source control by using the pattern supplied via ini
-	if ( FoundLabelString.IsEmpty() && InChangelistNumber >= 0 && !SourceControlBuildLabelPattern.IsEmpty() )
+	const double TotalTime = FPlatformTime::Seconds() - StartTime;
+	UE_LOG( LogCrashDebugHelper, Log, TEXT( "PDB Cache initialized in %.2f ms" ), TotalTime*1000.0f );
+	UE_LOG( LogCrashDebugHelper, Log, TEXT( "Found %i entries which occupy %i GBs" ), PDBCacheEntries.Num(), GetPDBCacheSizeGB() );
+}
+
+void FPDBCache::CleanPDBCache( int32 DaysToDelete, int32 NumberOfGBsToBeCleaned /*= 0 */ )
+{
+	// Not very efficient, but should do the trick.
+	// Revisit it later.
+	const double StartTime = FPlatformTime::Seconds();
+
+	TSet<FString> EntriesToBeRemoved;
+
+	// Find all outdated PDB Cache entries and mark them for removal.
+	const double DaysToDeleteAsSeconds = FTimespan( DaysToDelete, 0, 0, 0 ).GetTotalSeconds();
+	int32 NumGBsCleaned = 0;
+	for( const auto& It : PDBCacheEntries )
 	{
-		const FString ChangelistString = FString::Printf(TEXT("%d"), InChangelistNumber);
-		const FString TestLabel = SourceControlBuildLabelPattern.Replace(TEXT("%CHANGELISTNUMBER%"), *ChangelistString, ESearchCase::CaseSensitive );
-		TArray< TSharedRef<ISourceControlLabel> > Labels = ISourceControlModule::Get().GetProvider().GetLabels( TestLabel );
-		if ( Labels.Num() > 0 )
+		const FPDBCacheEntryRef& Entry = It.Value;
+		const FString EntryDirectory = PDBCachePath / Entry->Label;
+		const FString EntryTimeStampFilename = EntryDirectory / PDBTimeStampFile;
+
+		const double EntryFileAge = IFileManager::Get().GetFileAgeSeconds( *EntryTimeStampFilename );
+		if( EntryFileAge > DaysToDeleteAsSeconds )
 		{
-			// If we found more than one label, warn about it and just use the first one
-			if ( Labels.Num() > 1 )
+			EntriesToBeRemoved.Add( Entry->Label );
+			NumGBsCleaned += Entry->SizeGB;
+		}
+	}
+
+	if( NumberOfGBsToBeCleaned > 0 && NumGBsCleaned < NumberOfGBsToBeCleaned )
+	{
+		// Do the second pass if we need to remove more PDB Cache entries due to the free disk space restriction.
+		for( const auto& It : PDBCacheEntries )
+		{
+			const FPDBCacheEntryRef& Entry = It.Value;
+
+			if( !EntriesToBeRemoved.Contains(Entry->Label) )
 			{
-				UE_LOG(LogCrashDebugHelper, Warning, TEXT("RetrieveBuildLabel: More than one build label found with pattern %s - Using label %s"), *TestLabel, *Labels[0]->GetName());
+				EntriesToBeRemoved.Add( Entry->Label );
+				NumGBsCleaned += Entry->SizeGB;
+
+				if( NumGBsCleaned > NumberOfGBsToBeCleaned )
+				{
+					// Break the loop, we are done.
+					break;
+				}
 			}
+		}
+	}
 
-			FoundLabelString = Labels[0]->GetName();
-			UE_LOG(LogCrashDebugHelper, Log, TEXT("RetrieveBuildLabel: Failed to find build label in database %s.%s, but found label %s matching pattern %s in source control."), *DatabaseName, *DatabaseCatalog, *FoundLabelString, *TestLabel);
+	// Remove all marked PDB Cache entries.
+	for( const auto& EntryLabel : EntriesToBeRemoved )
+	{
+		RemovePDBCacheEntry( EntryLabel );
+	}
+
+	const double TotalTime = FPlatformTime::Seconds() - StartTime;
+	UE_LOG( LogCrashDebugHelper, Log, TEXT( "PDB Cache cleaned in %.2f ms" ), TotalTime*1000.0f );
+}
+
+FPDBCacheEntryRef FPDBCache::CreateAndAddPDBCacheEntry( const FString& OriginalLabelName, const FString& DepotName, const TArray<FString>& SyncedFiles )
+{
+	const FString CleanedLabelName = CleanLabelName( OriginalLabelName );
+	const FString EntryDirectory = PDBCachePath / CleanedLabelName;
+	const FString EntryTimeStampFilename = EntryDirectory / PDBTimeStampFile;
+
+	// Verify there is a entry timestamp file.
+	const bool bSaved = FFileHelper::SaveStringToFile( EntryTimeStampFilename, *EntryTimeStampFilename );
+	UE_CLOG( !bSaved, LogCrashDebugHelper, Fatal, TEXT( "Couldn't save the timestamp file to %s" ), *EntryTimeStampFilename );
+	const FDateTime LastAccessTime = IFileManager::Get().GetTimeStamp( *EntryTimeStampFilename );
+
+	// Copy all synced files to the PDC Cache entry directory.
+	// DepotName = //depot/UE4-Releases/4.1/
+	// OriginalLabelName = //depot/UE4-Releases/4.2/Rocket-CL-2082666
+	// CleanedLabelName = __depot_UE4-Releases_4.2_Rocket-CL-2082666
+	// SyncedFiles = { //depot/UE4-Releases/4.1/engine/binaries/dotnet/swarminterface*.dll, ... }
+	// CWD = F:\depot\UE4-Releases\4.1\Engine\Binaries\Win64
+	// RootDir = F:\depot\UE4-Releases\4.1\
+
+#define DO_LOCAL_TESTING 1
+#if	DO_LOCAL_TESTING
+	const FString RootDir = TEXT( "U:/P4EPIC/UE4-Releases/4.1/" );
+#else
+	const FString RootDir = FPaths::RootDir();
+#endif // DO_LOCAL_TESTING
+
+	UE_LOG( LogCrashDebugHelper, Warning, TEXT( "PDB Cache entry %s is being copied from %s, it will take some time" ), *CleanedLabelName, *OriginalLabelName );
+	for( const auto& Filename : SyncedFiles )
+	{
+		const FString SourceDirectoryWithSearch = Filename.Replace( *DepotName, *RootDir );
+
+		TArray<FString> MatchedFiles;
+		IFileManager::Get().FindFiles( MatchedFiles, *SourceDirectoryWithSearch, true, false );
+
+		for( const auto& MatchedFilename : MatchedFiles )
+		{
+			const FString SrcFilename = FPaths::GetPath( SourceDirectoryWithSearch ) / MatchedFilename;
+			const FString DestFilename = EntryDirectory / SrcFilename.Replace( *RootDir, TEXT("") );
+			IFileManager::Get().Copy( *DestFilename, *SrcFilename );
 		}
 	}
 
-	return FoundLabelString;
+
+	TArray<FString> PDBFiles;
+	IFileManager::Get().FindFilesRecursive( PDBFiles, *EntryDirectory, TEXT( "*.*" ), true, false );
+
+	// Calculate the size of this PDB Cache entry.
+	int64 TotalSize = 0;
+	for( const auto& Filename : PDBFiles )
+	{
+		const int64 FileSize = IFileManager::Get().FileSize( *Filename );
+		TotalSize += FileSize;
+	}
+
+	// Round-up the size.
+	const int32 SizeGB = (int32)FMath::DivideAndRoundUp( TotalSize, (int64)NUM_BYTES_PER_GB );
+
+	FPDBCacheEntryRef NewCacheEntry = MakeShareable( new FPDBCacheEntry( CleanedLabelName, SizeGB ) );
+	NewCacheEntry->LastAccessTime = LastAccessTime;
+	NewCacheEntry->Files = PDBFiles;
+	
+	PDBCacheEntries.Add( CleanedLabelName, NewCacheEntry );
+	SortPDBCache();
+
+	return NewCacheEntry;
+}
+
+FPDBCacheEntryRef FPDBCache::ReadPDBCacheEntry( const FString& InLabel )
+{
+	const FString EntryDirectory = PDBCachePath / InLabel;
+	const FString EntryTimeStampFilename = EntryDirectory / PDBTimeStampFile;
+
+	// Verify there is a entry timestamp file.
+	const FDateTime LastAccessTime = IFileManager::Get().GetTimeStamp( *EntryTimeStampFilename );
+
+	TArray<FString> PDBFiles;
+	IFileManager::Get().FindFilesRecursive( PDBFiles, *EntryDirectory, TEXT( "*.*" ), true, false );
+
+	// Calculate the size of this PDB Cache entry.
+	int64 TotalSize = 0;
+	for( const auto& Filename : PDBFiles )
+	{
+		const int64 FileSize = IFileManager::Get().FileSize( *Filename );
+		TotalSize += FileSize;
+	}
+
+	// Round-up the size.
+	const int32 SizeGB = (int32)FMath::DivideAndRoundUp( TotalSize, (int64)NUM_BYTES_PER_GB );
+
+	FPDBCacheEntryRef NewEntry = MakeShareable( new FPDBCacheEntry( InLabel, SizeGB ) );
+	NewEntry->LastAccessTime = LastAccessTime;
+	NewEntry->Files = PDBFiles;
+	return NewEntry;
+}
+
+void FPDBCache::TouchPDBCacheEntry( const FString& InLabel )
+{
+	const FString CleanedLabelName = CleanLabelName( InLabel );
+	const FString EntryDirectory = PDBCachePath / CleanedLabelName;
+	const FString EntryTimeStampFilename = EntryDirectory / PDBTimeStampFile;
+
+	FPDBCacheEntryRef Entry = PDBCacheEntries.FindChecked( CleanedLabelName );
+	Entry->SetLastAccessTimeToNow();
+
+	const bool bResult = IFileManager::Get().SetTimeStamp( *EntryTimeStampFilename, Entry->LastAccessTime );
+	SortPDBCache();
 }
 
+void FPDBCache::RemovePDBCacheEntry( const FString& InLabel )
+{
+	const double StartTime = FPlatformTime::Seconds();
+
+	const FString EntryDirectory = PDBCachePath / InLabel;
+
+	FPDBCacheEntryRef& Entry = PDBCacheEntries.FindChecked( InLabel );
+	IFileManager::Get().DeleteDirectory( *EntryDirectory, true, true );
+	PDBCacheEntries.Remove( InLabel );
+
+	const double TotalTime = FPlatformTime::Seconds() - StartTime;
+	UE_LOG( LogCrashDebugHelper, Warning, TEXT( "PDB Cache entry %s removed in %.2f ms, restored %i GBs" ), *InLabel, TotalTime*1000.0f, Entry->SizeGB );
+}
+
+
+
diff --git a/Engine/Source/Developer/CrashDebugHelper/Private/Windows/WindowsPlatformStackWalkExt.cpp b/Engine/Source/Developer/CrashDebugHelper/Private/Windows/WindowsPlatformStackWalkExt.cpp
--- a/Engine/Source/Developer/CrashDebugHelper/Private/Windows/WindowsPlatformStackWalkExt.cpp
+++ b/Engine/Source/Developer/CrashDebugHelper/Private/Windows/WindowsPlatformStackWalkExt.cpp
@@ -180,25 +180,60 @@ void FWindowsPlatformStackWalkExt::GetModuleList( FCrashInfo* CrashInfo )
  */
 void FWindowsPlatformStackWalkExt::SetSymbolPathsFromModules( FCrashInfo* CrashInfo )
 {
-	TArray<FString> SymbolPaths;
-	for( int32 ModuleNameIndex = 0; ModuleNameIndex < CrashInfo->ModuleNames.Num(); ModuleNameIndex++ )
+	const bool bUseCachedData = CrashInfo->PDBCacheEntry.IsValid();
+	FString CombinedPath = TEXT( "" );
+
+	if( bUseCachedData )
 	{
-		FString ModulePath = FPaths::GetPath( CrashInfo->ModuleNames[ModuleNameIndex] );
-		if( ModulePath.Len() > 0 )
+		TSet<FString> SymbolPaths;
+		for( const auto& Filename : CrashInfo->PDBCacheEntry->Files )
 		{
-			SymbolPaths.AddUnique( ModulePath );
+			const FString SymbolPath = FPaths::GetPath( Filename );
+			if( SymbolPath.Len() > 0 )
+			{
+				SymbolPaths.Add( SymbolPath );
+			}
 		}
-	}
 
-	FString CombinedPath = TEXT( "" );
-	for( int32 PathIndex = 0; PathIndex < SymbolPaths.Num(); PathIndex++ )
-	{
-		CombinedPath += FString( TEXT( "../../../" ) ) + SymbolPaths[PathIndex] + TEXT( ";" );
+		for( const auto& SymbolPath : SymbolPaths )
+		{
+			CombinedPath += SymbolPath;
+			CombinedPath += TEXT( ";" );
+		}
+
+		// Set the symbol path
+		Symbol->SetImagePathWide( *CombinedPath );
+		Symbol->SetSymbolPathWide( *CombinedPath );
 	}
+	else
+	{
+		TArray<FString> SymbolPaths;
+		for( int32 ModuleNameIndex = 0; ModuleNameIndex < CrashInfo->ModuleNames.Num(); ModuleNameIndex++ )
+		{
+			FString ModulePath = FPaths::GetPath( CrashInfo->ModuleNames[ModuleNameIndex] );
+			if( ModulePath.Len() > 0 )
+			{
+				SymbolPaths.AddUnique( ModulePath );
+			}
+		}
 
-	// Set the symbol path
-	Symbol->SetImagePathWide( *CombinedPath );
-	Symbol->SetSymbolPathWide( *CombinedPath );
+		for( int32 PathIndex = 0; PathIndex < SymbolPaths.Num(); PathIndex++ )
+		{
+			CombinedPath += FString( TEXT( "../../../" ) ) + SymbolPaths[PathIndex] + TEXT( ";" );
+		}
+
+		// Set the symbol path
+		Symbol->SetImagePathWide( *CombinedPath );
+		Symbol->SetSymbolPathWide( *CombinedPath );
+	
+		//@TODO: ROCKETHACK: This is for Rocket only.
+		Symbol->AppendImagePathWide( TEXT( "..\\..\\..\\Rocket\\Installed\\Windows\\Engine\\Binaries\\Win64" ) );
+		Symbol->AppendImagePathWide( TEXT( "..\\..\\..\\Rocket\\Installed\\Windows\\Engine\\Binaries\\Win32" ) );
+		Symbol->AppendSymbolPathWide( TEXT( "..\\..\\..\\Rocket\\Symbols\\Engine\\Binaries\\Win64" ) );
+		Symbol->AppendSymbolPathWide( TEXT( "..\\..\\..\\Rocket\\Symbols\\Engine\\Binaries\\Win32" ) );
+		Symbol->AppendImagePathWide( TEXT( "..\\..\\..\\Rocket\\LauncherInstalled\\Windows\\Launcher\\Engine\\Binaries\\Win64" ) );
+		Symbol->AppendSymbolPathWide( TEXT( "..\\..\\..\\Rocket\\LauncherSymbols\\Windows\\Launcher\\Engine\\Binaries\\Win64" ) );
+	}
 
 	// Add in syncing of the Microsoft symbol servers if requested
 	if( FParse::Param( FCommandLine::Get(), TEXT( "SyncMicrosoftSymbols" ) ) )
@@ -221,14 +256,6 @@ void FWindowsPlatformStackWalkExt::SetSymbolPathsFromModules( FCrashInfo* CrashI
 		Symbol->AppendSymbolPathWide( TEXT( "SRV*..\\..\\Intermediate\\SymbolCache*http://msdl.microsoft.com/download/symbols" ) );
 	}
 
-	//@TODO: ROCKETHACK: This is for Rocket only.
-	Symbol->AppendImagePathWide( TEXT( "..\\..\\..\\Rocket\\Installed\\Windows\\Engine\\Binaries\\Win64" ) );
-	Symbol->AppendImagePathWide( TEXT( "..\\..\\..\\Rocket\\Installed\\Windows\\Engine\\Binaries\\Win32" ) );
-	Symbol->AppendSymbolPathWide( TEXT( "..\\..\\..\\Rocket\\Symbols\\Engine\\Binaries\\Win64" ) );
-	Symbol->AppendSymbolPathWide( TEXT( "..\\..\\..\\Rocket\\Symbols\\Engine\\Binaries\\Win32" ) );
-	Symbol->AppendImagePathWide( TEXT( "..\\..\\..\\Rocket\\LauncherInstalled\\Windows\\Launcher\\Engine\\Binaries\\Win64" ) );
-	Symbol->AppendSymbolPathWide( TEXT( "..\\..\\..\\Rocket\\LauncherSymbols\\Windows\\Launcher\\Engine\\Binaries\\Win64" ) );
-
 	TCHAR SymbolPath[1024] = { 0 };
 	Symbol->GetSymbolPathWide( SymbolPath, 1024, NULL );
 	CrashInfo->Log( FString::Printf( TEXT( " ... setting symbol path to: '%s'" ), SymbolPath ) );
diff --git a/Engine/Source/Developer/CrashDebugHelper/Public/CrashDebugHelper.h b/Engine/Source/Developer/CrashDebugHelper/Public/CrashDebugHelper.h
--- a/Engine/Source/Developer/CrashDebugHelper/Public/CrashDebugHelper.h
+++ b/Engine/Source/Developer/CrashDebugHelper/Public/CrashDebugHelper.h
@@ -2,6 +2,10 @@
 
 #pragma once
 
+struct FPDBCacheEntry;
+typedef TSharedRef<FPDBCacheEntry> FPDBCacheEntryRef;
+typedef TSharedPtr<FPDBCacheEntry> FPDBCacheEntryPtr;
+
 enum EProcessorArchitecture
 {
 	PA_UNKNOWN,
@@ -135,6 +139,9 @@ public:
 	TArray<FCrashThreadInfo> Threads;
 	TArray<FCrashModuleInfo> Modules;
 
+	/** Shared pointer to the PDB Cache entry, if valid contains all information about synced PDBs. */
+	FPDBCacheEntryPtr PDBCacheEntry;
+
 	FCrashInfo()
 		: ChangelistBuiltFrom( -1 )
 		, SourceLineNumber( 0 )
@@ -164,7 +171,6 @@ private:
 /** Helper structure for tracking crash debug information */
 struct CRASHDEBUGHELPER_API FCrashDebugInfo
 {
-public:
 	/** The name of the crash dump file */
 	FString CrashDumpName;
 	/** The engine version of the crash dump build */
@@ -175,6 +181,206 @@ public:
 	FString SourceControlLabel;
 };
 
+/** Helper struct that holds various information about one PDB Cache entry. */
+struct FPDBCacheEntry
+{
+	/** Initialization constructor. */
+	FPDBCacheEntry( const FString& InLabel, const int32 InSizeGB )
+		: Label( InLabel )
+		, SizeGB( InSizeGB )
+	{}
+
+	void SetLastAccessTimeToNow()
+	{
+		LastAccessTime = FDateTime::UtcNow();
+	}
+
+	/** Paths to files associated with this PDB Cache entry. */
+	TArray<FString> Files;
+
+	/** P4 Label associated with this PDB Cache entry. */
+	const FString Label;
+	
+	/** Last access time, changed every time this PDB cache entry is used. */
+	FDateTime LastAccessTime;
+
+	/** Size of the cache entry, in GBs. Rounded-up. */
+	const int32 SizeGB;
+};
+
+struct FPDBCacheEntryByAccessTime
+{
+	FORCEINLINE bool operator()( const FPDBCacheEntryRef& A, const FPDBCacheEntryRef& B ) const
+	{
+		return A->LastAccessTime.GetTicks() < B->LastAccessTime.GetTicks();
+	}
+};
+
+/** Implements PDB cache. */
+struct FPDBCache
+{
+	//friend class ICrashDebugHelper;
+protected:
+
+	// Defaults for the PDB cache.
+	enum
+	{
+		/** Size of the PDB cache, in GBs. */
+		PDB_CACHE_SIZE_GB = 128,
+		MIN_FREESPACE_GB = 64,
+		/** Age of file when it should be deleted from the PDB cache. */
+		DAYS_TO_DELETE_UNUSED_FILES = 14,
+
+		/**
+		*	Number of iterations inside the CleanPDBCache method.
+		*	Mostly to verify that MinDiskFreeSpaceGB requirement is met.
+		*/
+		CLEAN_PDBCACHE_NUM_ITERATIONS = 2,
+
+		/** Number of bytes per one gigabyte. */
+		NUM_BYTES_PER_GB = 1024 * 1024 * 1024
+	};
+
+	/** Dummy file used to read/set the file timestamp. */
+	static const TCHAR* PDBTimeStampFile;
+
+	/** Map of the PDB Cache entries. */
+	TMap<FString, FPDBCacheEntryRef> PDBCacheEntries;
+
+	/** Path to the folder where the PDB cache is stored. */
+	FString PDBCachePath;
+
+	/** The builtFromCL to use instead of extracting from the minidump */
+	int32 BuiltFromCL;
+
+	/** Age of file when it should be deleted from the PDB cache. */
+	int32 DaysToDeleteUnusedFilesFromPDBCache;
+
+	/** Size of the PDB cache, in GBs. */
+	int32 PDBCacheSizeGB;
+
+	/**
+	*	Minimum disk free space that should be available on the disk where the PDB cache is stored, in GBs.
+	*	Minidump diagnostics runs usually on the same drive as the crash reports drive, so we need to leave some space for the crash receiver.
+	*	If there is not enough disk free space, we will run the clean-up process.
+	*/
+	int32 MinDiskFreeSpaceGB;
+
+	/** Whether to use the PDB cache. */
+	bool bUsePDBCache;
+
+public:
+	/** Default constructor. */
+	FPDBCache()
+		: BuiltFromCL( 0 )
+		, DaysToDeleteUnusedFilesFromPDBCache( DAYS_TO_DELETE_UNUSED_FILES )
+		, PDBCacheSizeGB( PDB_CACHE_SIZE_GB )
+		, MinDiskFreeSpaceGB( MIN_FREESPACE_GB )
+		, bUsePDBCache( false )
+	{}
+
+	/** Basic initialization, reading config etc.. */
+	void Init();
+
+	/**
+	 * @return whether to use the PDB cache.
+	 */
+	bool UsePDBCache() const
+	{
+		return bUsePDBCache;
+	}
+	
+	/**
+	 * @return true, if the PDB Cache contains the specified label.
+	 */
+	bool ContainsPDBCacheEntry( const FString& InLabel ) const
+	{
+		return PDBCacheEntries.Contains( CleanLabelName( InLabel ) );
+	}
+
+	/**
+	 *	Touches a PDB Cache entry by updating the timestamp.
+	 */
+	void TouchPDBCacheEntry( const FString& InLabel );
+
+	/**
+	 * @return a PDB Cache entry for the specified label.
+	 */
+	FPDBCacheEntryRef FindPDBCacheEntry( const FString& InLabel )
+	{
+		return PDBCacheEntries.FindChecked( CleanLabelName( InLabel ) );
+	}
+
+	/**
+	 *	Creates a new PDB Cache entry, initializes it and adds to the database.
+	 */
+	FPDBCacheEntryRef CreateAndAddPDBCacheEntry( const FString& OriginalLabelName, const FString& DepotName, const TArray<FString>& SyncedFiles );
+
+protected:
+	/** Initializes the PDB Cache. */
+	void InitializePDBCache();
+
+	/**
+	 * @return the size of the PDB cache entry, in GBs.
+	 */
+	int32 GetPDBCacheEntrySizeGB( const FString& InLabel ) const
+	{
+		return PDBCacheEntries.FindChecked( InLabel )->SizeGB;
+	}
+
+	/**
+	 * @return the size of the PDB cache directory, in GBs.
+	 */
+	int32 GetPDBCacheSizeGB() const
+	{
+		int32 Result = 0;
+		if( bUsePDBCache )
+		{
+			for( const auto& It : PDBCacheEntries )
+			{
+				Result += It.Value->SizeGB;
+			}
+		}
+		return Result;
+	}
+
+	/**
+	 * Cleans the PDB Cache.
+	 *
+	 * @param DaysToKeep - removes all PDB Cache entries that are older than this value
+	 * @param NumberOfGBsToBeCleaned - if specifies, will try to remove as many PDB Cache entries as needed
+	 * to free disk space specified by this value
+	 *
+	 */
+	void CleanPDBCache( int32 DaysToKeep, int32 NumberOfGBsToBeCleaned = 0 );
+
+
+	/**
+	 *	Reads an existing PDB Cache entry.
+	 */
+	FPDBCacheEntryRef ReadPDBCacheEntry( const FString& InLabel );
+
+	/**
+	 *	Sort PDB Cache entries by last access time.
+	 */
+	void SortPDBCache()
+	{
+		PDBCacheEntries.ValueSort( FPDBCacheEntryByAccessTime() );
+	}
+
+	/**
+	 *	Removes a PDB Cache entry from the database.
+	 *	Also removes all external files associated with this PDB Cache entry.
+	 */
+	void RemovePDBCacheEntry( const FString& InLabel );
+
+	/** Replaces all / chars with _ for the specified label name. */
+	FString CleanLabelName( const FString& InLabel ) const
+	{
+		return InLabel.Replace( TEXT( "/" ), TEXT( "_" ) );
+	}
+};
+
 /** The public interface for the crash dump handler singleton. */
 class CRASHDEBUGHELPER_API ICrashDebugHelper
 {
@@ -183,26 +389,24 @@ protected:
 	FString DepotName;
 	/** The local folder where symbol files should be stored */
 	FString LocalSymbolStore;
-	/** The database that the builder version information is stored in */
-	FString DatabaseName;
-	/** The catalog that the builder version info is stored in */
-	FString DatabaseCatalog;
 	/** In the event that a database query for a label fails, use this pattern to search in source control for the label */
 	FString SourceControlBuildLabelPattern;
 	/** The builtFromCL to use instead of extracting from the minidump */
 	int32 BuiltFromCL;
-
+	
 	/** Indicates that the crash handler is ready to do work */
 	bool bInitialized;
 
+	/** Instance of the PDB Cache. */
+	FPDBCache PDBCache;
+
 public:
 	/** A platform independent representation of a crash */
 	FCrashInfo CrashInfo;
-
+	
 	/** Virtual destructor */
 	virtual ~ICrashDebugHelper()
-	{
-	};
+	{}
 
 	/**
 	 *	Initialize the helper
@@ -211,18 +415,6 @@ public:
 	 */
 	virtual bool Init();
 
-	/** Get the database name */
-	const FString& GetDatabaseName() const
-	{
-		return DatabaseName;
-	}
-
-	/** Get the catalog name */
-	const FString& GetDatabaseCatalog() const
-	{
-		return DatabaseCatalog;
-	}
-
 	/** Get the depot name */
 	const FString& GetDepotName() const
 	{
