CrashDebugHelper - Implemented better method for getting the list of required pdb/dll and setting the symbol paths for modules

diff --git a/Engine/Source/Developer/CrashDebugHelper/Private/CrashDebugHelper.cpp b/Engine/Source/Developer/CrashDebugHelper/Private/CrashDebugHelper.cpp
--- a/Engine/Source/Developer/CrashDebugHelper/Private/CrashDebugHelper.cpp
+++ b/Engine/Source/Developer/CrashDebugHelper/Private/CrashDebugHelper.cpp
@@ -131,9 +131,12 @@ bool ICrashDebugHelper::SyncModules( FCrashInfo* CrashInfo )
 		return false;
 	}
 
-	// Sync all the dll, exes, and related symbol files
+	// Get all labels associated with the crash info's label.
+	// It should return one unique label because the crash info's label looks like this //depot/UE4-Releases/4.0/Rocket-CL-2047835 @see RetrieveBuildLabel
+	// We can assume that the label name is unique.
+	// 
 	TArray< TSharedRef<ISourceControlLabel> > Labels = ISourceControlModule::Get().GetProvider().GetLabels( CrashInfo->LabelName );
-	if(Labels.Num() > 0)
+	if( Labels.Num() == 1 )
 	{
 		const bool bContains = PDBCache.UsePDBCache() && PDBCache.ContainsPDBCacheEntry( CrashInfo->LabelName );
 		if( bContains )
@@ -145,23 +148,87 @@ bool ICrashDebugHelper::SyncModules( FCrashInfo* CrashInfo )
 		}
 		else
 		{
-			TArray<FString> SyncedFiles;
+			TArray<FString> SyncedFiles;	
+			TSharedRef<ISourceControlLabel> Label = Labels[0];
 
-			// Sync every module from every label. If the same modules appear in every label, this will fail.
-			for( auto LabelIt = Labels.CreateConstIterator(); LabelIt; ++LabelIt )
+			//@TODO: MAC: Excluding labels for Mac since we are only syncing windows binaries here...
+			if( Label->GetName().Contains( TEXT( "Mac" ) ) )
 			{
-				TSharedRef<ISourceControlLabel> Label = *LabelIt;
+				UE_LOG( LogCrashDebugHelper, Log, TEXT( " Skipping Mac label '%s' when syncing modules." ), *Label->GetName() );
+			}
+			else
+			{
+				// Sync all the dll, exes, and related symbol files
+				UE_LOG( LogCrashDebugHelper, Log, TEXT( " Syncing modules with label '%s'." ), *Label->GetName() );
+
+				TSet<FString> FilesToSync;
+
+				const bool bUseNewSyncModule = true;
 
-				//@TODO: MAC: Excluding labels for Mac since we are only syncing windows binaries here...
-				if( Label->GetName().Contains( TEXT( "Mac" ) ) )
+				if( bUseNewSyncModule )
 				{
-					UE_LOG( LogCrashDebugHelper, Log, TEXT( " Skipping Mac label '%s' when syncing modules." ), *Label->GetName() );
+					SCOPE_LOG_TIME_IN_SECONDS( TEXT( "SyncModules_NEW" ), nullptr );
+
+					// Grab all dll and pdb files for the specified label.
+					TArray< TSharedRef<class ISourceControlRevision, ESPMode::ThreadSafe> > DLLSourceControlRevisions;
+					const FString DLLsPath = FString::Printf( TEXT( "%s/....dll*" ), *DepotName );
+					Label->GetFileRevisions( DLLsPath, DLLSourceControlRevisions );
+
+					TArray< TSharedRef<class ISourceControlRevision, ESPMode::ThreadSafe> > PDBSourceControlRevisions;
+					const FString PDBsPath = FString::Printf( TEXT( "%s/....pdb*" ), *DepotName );
+					Label->GetFileRevisions( PDBsPath, PDBSourceControlRevisions );		
+			
+					TSet<FString> DLLPaths;
+					for( const auto& DLLSrc : DLLSourceControlRevisions )
+					{
+						DLLPaths.Add( DLLSrc->GetFilename().Replace( *(DepotName + TEXT( "/" )), TEXT( "" ) ) );
+					}
+
+					TSet<FString> PDBPaths;
+					for( const auto& PDBSrc : PDBSourceControlRevisions )
+					{
+						PDBPaths.Add( PDBSrc->GetFilename().Replace( *(DepotName + TEXT( "/" )), TEXT( "" ) ) );
+					}
+
+					// Iterate through all module and see if we have dll and pdb associated with the module, if so add it to the files to sync.
+					for( const auto& ModuleNameDLL : CrashInfo->ModuleNames )
+					{
+						const FString ModuleNamePDB = ModuleNameDLL.Replace( TEXT( ".dll" ), TEXT( ".pdb" ) );
+
+						for( const auto& DLLPath : DLLPaths )
+						{
+							const bool bContainsDLL = DLLPath.Contains( ModuleNameDLL );
+							if( bContainsDLL )
+							{
+								FilesToSync.Add( DLLPath );
+							}
+						}
+
+						for( const auto& PDBPath : PDBPaths )
+						{
+							const bool bContainsPDB = PDBPath.Contains( ModuleNamePDB );
+							if( bContainsPDB )
+							{
+								FilesToSync.Add( PDBPath );
+							}
+						}
+					}
+					
+					// Now, sync all files.
+					for( const auto& Filename : FilesToSync )
+					{
+						const FString DepotPath = FString::Printf( TEXT( "%s/%s" ), *DepotName, *Filename );
+						if( Label->Sync( DepotPath ) )
+						{
+							UE_LOG( LogCrashDebugHelper, Warning, TEXT( " ... synced binary '%s'." ), *DepotPath );
+						}
+						SyncedFiles.Add( DepotPath );
+					}
 				}
 				else
 				{
-					UE_LOG( LogCrashDebugHelper, Log, TEXT( " Syncing modules with label '%s'." ), *Label->GetName() );
+					SCOPE_LOG_TIME_IN_SECONDS( TEXT( "SyncModules_OLD" ), nullptr );
 
-					TArray<FString> FilesToSync;
 					for( int32 ModuleNameIndex = 0; ModuleNameIndex < CrashInfo->ModuleNames.Num(); ModuleNameIndex++ )
 					{
 						// Sometimes P4 may hang due to spamming the server with the commands, if so try uncommenting the line below.
@@ -170,60 +237,57 @@ bool ICrashDebugHelper::SyncModules( FCrashInfo* CrashInfo )
 						// Match all decorated versions of the module. We may need them if the file was renamed to remove the "-Platform-Configuration" decoration.
 						DepotPath = DepotPath.Replace( TEXT( ".dll" ), TEXT( "*.dll" ) ).Replace( TEXT( ".exe" ), TEXT( "*.exe" ) );
 
+						if( Label->Sync( DepotPath ) )
 						{
-							TSharedRef<ISourceControlLabel> Label = Labels[0];
-							if( Label->Sync( DepotPath ) )
-							{
-								UE_LOG( LogCrashDebugHelper, Warning, TEXT( " ... synced binary '%s'." ), *DepotPath );
-							}
-							SyncedFiles.Add( DepotPath );
+							UE_LOG( LogCrashDebugHelper, Warning, TEXT( " ... synced binary '%s'." ), *DepotPath );
+						}
+						SyncedFiles.Add( DepotPath );
 
-				        	FString PDBName = DepotPath.Replace( TEXT( ".dll" ), TEXT( ".pdb" ) ).Replace( TEXT( ".exe" ), TEXT( ".pdb" ) );
-							if( Label->Sync( PDBName ) )
-							{
-								UE_LOG( LogCrashDebugHelper, Warning, TEXT( " ... synced symbol '%s'." ), *PDBName );
-							}
-							SyncedFiles.Add( PDBName );
+						FString PDBName = DepotPath.Replace( TEXT( ".dll" ), TEXT( ".pdb" ) ).Replace( TEXT( ".exe" ), TEXT( ".pdb" ) );
+						if( Label->Sync( PDBName ) )
+						{
+							UE_LOG( LogCrashDebugHelper, Warning, TEXT( " ... synced symbol '%s'." ), *PDBName );
+						}
+						SyncedFiles.Add( PDBName );
 
-							//@TODO: ROCKETHACK: Adding additional Installed and Symbol paths - revisit when builds are made by the builder...
-							DepotPath = FString::Printf( TEXT( "%s/Rocket/Installed/Windows/%s" ), *DepotName, *CrashInfo->ModuleNames[ModuleNameIndex] );
-							// Match all decorated versions of the module. We may need them if the file was renamed to remove the "-Platform-Configuration" decoration.
-							DepotPath = DepotPath.Replace( TEXT( ".dll" ), TEXT( "*.dll" ) ).Replace( TEXT( ".exe" ), TEXT( "*.exe" ) );
-							if( Label->Sync( DepotPath ) )
-							{
-								UE_LOG( LogCrashDebugHelper, Warning, TEXT( " ... synced binary '%s'." ), *DepotPath );
-							}
-							SyncedFiles.Add( DepotPath );
+						//@TODO: ROCKETHACK: Adding additional Installed and Symbol paths - revisit when builds are made by the builder...
+						DepotPath = FString::Printf( TEXT( "%s/Rocket/Installed/Windows/%s" ), *DepotName, *CrashInfo->ModuleNames[ModuleNameIndex] );
+						// Match all decorated versions of the module. We may need them if the file was renamed to remove the "-Platform-Configuration" decoration.
+						DepotPath = DepotPath.Replace( TEXT( ".dll" ), TEXT( "*.dll" ) ).Replace( TEXT( ".exe" ), TEXT( "*.exe" ) );
+						if( Label->Sync( DepotPath ) )
+						{
+							UE_LOG( LogCrashDebugHelper, Warning, TEXT( " ... synced binary '%s'." ), *DepotPath );
+						}
+						SyncedFiles.Add( DepotPath );
 
-							DepotPath = FString::Printf( TEXT( "%s/Rocket/Symbols/%s" ), *DepotName, *CrashInfo->ModuleNames[ModuleNameIndex] );
-							PDBName = DepotPath.Replace( TEXT( ".dll" ), TEXT( "*.pdb" ) ).Replace( TEXT( ".exe" ), TEXT( "*.pdb" ) );
-							if( Label->Sync( PDBName ) )
-							{
-								UE_LOG( LogCrashDebugHelper, Warning, TEXT( " ... synced symbol '%s'." ), *PDBName );		
-							}
-							SyncedFiles.Add( PDBName );
+						DepotPath = FString::Printf( TEXT( "%s/Rocket/Symbols/%s" ), *DepotName, *CrashInfo->ModuleNames[ModuleNameIndex] );
+						PDBName = DepotPath.Replace( TEXT( ".dll" ), TEXT( "*.pdb" ) ).Replace( TEXT( ".exe" ), TEXT( "*.pdb" ) );
+						if( Label->Sync( PDBName ) )
+						{
+							UE_LOG( LogCrashDebugHelper, Warning, TEXT( " ... synced symbol '%s'." ), *PDBName );
+						}
+						SyncedFiles.Add( PDBName );
 
-							DepotPath = FString::Printf( TEXT( "%s/Rocket/LauncherInstalled/Windows/Launcher/%s" ), *DepotName, *CrashInfo->ModuleNames[ModuleNameIndex] );
-							// Match all decorated versions of the module. We may need them if the file was renamed to remove the "-Platform-Configuration" decoration.
-							DepotPath = DepotPath.Replace( TEXT( ".dll" ), TEXT( "*.dll" ) ).Replace( TEXT( ".exe" ), TEXT( "*.exe" ) );
-							if( Label->Sync( DepotPath ) )
-							{
-								UE_LOG( LogCrashDebugHelper, Warning, TEXT( " ... synced binary '%s'." ), *DepotPath );
-							}
-							SyncedFiles.Add( DepotPath );
+						DepotPath = FString::Printf( TEXT( "%s/Rocket/LauncherInstalled/Windows/Launcher/%s" ), *DepotName, *CrashInfo->ModuleNames[ModuleNameIndex] );
+						// Match all decorated versions of the module. We may need them if the file was renamed to remove the "-Platform-Configuration" decoration.
+						DepotPath = DepotPath.Replace( TEXT( ".dll" ), TEXT( "*.dll" ) ).Replace( TEXT( ".exe" ), TEXT( "*.exe" ) );
+						if( Label->Sync( DepotPath ) )
+						{
+							UE_LOG( LogCrashDebugHelper, Warning, TEXT( " ... synced binary '%s'." ), *DepotPath );
+						}
+						SyncedFiles.Add( DepotPath );
 
-							DepotPath = FString::Printf( TEXT( "%s/Rocket/LauncherSymbols/Windows/Launcher/%s" ), *DepotName, *CrashInfo->ModuleNames[ModuleNameIndex] );
-							PDBName = DepotPath.Replace( TEXT( ".dll" ), TEXT( "*.pdb" ) ).Replace( TEXT( ".exe" ), TEXT( "*.pdb" ) );
-							if( Label->Sync( PDBName ) )
-							{
-								UE_LOG( LogCrashDebugHelper, Warning, TEXT( " ... synced symbol '%s'." ), *PDBName );
-							}
-							SyncedFiles.Add( PDBName );
+						DepotPath = FString::Printf( TEXT( "%s/Rocket/LauncherSymbols/Windows/Launcher/%s" ), *DepotName, *CrashInfo->ModuleNames[ModuleNameIndex] );
+						PDBName = DepotPath.Replace( TEXT( ".dll" ), TEXT( "*.pdb" ) ).Replace( TEXT( ".exe" ), TEXT( "*.pdb" ) );
+						if( Label->Sync( PDBName ) )
+						{
+							UE_LOG( LogCrashDebugHelper, Warning, TEXT( " ... synced symbol '%s'." ), *PDBName );
 						}
+						SyncedFiles.Add( PDBName );
 					}
 				}
 			}
-
+			
 			if( PDBCache.UsePDBCache() )
 			{
 				// Initialize and add a new PDB Cache entry to the database.
diff --git a/Engine/Source/Developer/CrashDebugHelper/Private/Windows/WindowsPlatformStackWalkExt.cpp b/Engine/Source/Developer/CrashDebugHelper/Private/Windows/WindowsPlatformStackWalkExt.cpp
--- a/Engine/Source/Developer/CrashDebugHelper/Private/Windows/WindowsPlatformStackWalkExt.cpp
+++ b/Engine/Source/Developer/CrashDebugHelper/Private/Windows/WindowsPlatformStackWalkExt.cpp
@@ -207,32 +207,69 @@ void FWindowsPlatformStackWalkExt::SetSymbolPathsFromModules( FCrashInfo* CrashI
 	}
 	else
 	{
-		TArray<FString> SymbolPaths;
-		for( int32 ModuleNameIndex = 0; ModuleNameIndex < CrashInfo->ModuleNames.Num(); ModuleNameIndex++ )
+		const bool bUseNewSetPath = true;
+		if( bUseNewSetPath )
 		{
-			FString ModulePath = FPaths::GetPath( CrashInfo->ModuleNames[ModuleNameIndex] );
-			if( ModulePath.Len() > 0 )
+//#define DO_LOCAL_TESTING 1
+#if	DO_LOCAL_TESTING
+			const FString RootDir = TEXT( "U:/P4EPIC/UE4-Releases/4.1/" );
+#else
+			const FString RootDir = FPaths::RootDir();
+#endif // DO_LOCAL_TESTING
+
+			TArray<FString> ModulesAndSymbols;
+			IFileManager::Get().FindFilesRecursive( ModulesAndSymbols, *RootDir, TEXT( "*.pdb" ), true, false, false );
+			IFileManager::Get().FindFilesRecursive( ModulesAndSymbols, *RootDir, TEXT( "*.dll" ), true, false, false );
+
+			TSet<FString> SymbolPaths;
+			for( const auto& Filename : ModulesAndSymbols )
 			{
-				SymbolPaths.AddUnique( ModulePath );
+				const FString Path = FPaths::GetPath( Filename );
+				if( Path.Len() > 0 )
+				{
+					SymbolPaths.Add( Path );
+				}
 			}
-		}
 
-		for( int32 PathIndex = 0; PathIndex < SymbolPaths.Num(); PathIndex++ )
-		{
-			CombinedPath += FString( TEXT( "../../../" ) ) + SymbolPaths[PathIndex] + TEXT( ";" );
+			for( const auto& SymbolPath : SymbolPaths )
+			{
+				CombinedPath += SymbolPath;
+				CombinedPath += TEXT( ";" );
+			}
+
+			// Set the symbol path
+			Symbol->SetImagePathWide( *CombinedPath );
+			Symbol->SetSymbolPathWide( *CombinedPath );
 		}
+		else
+		{
+			TArray<FString> SymbolPaths;
+			for( int32 ModuleNameIndex = 0; ModuleNameIndex < CrashInfo->ModuleNames.Num(); ModuleNameIndex++ )
+			{
+				FString ModulePath = FPaths::GetPath( CrashInfo->ModuleNames[ModuleNameIndex] );
+				if( ModulePath.Len() > 0 )
+				{
+					SymbolPaths.AddUnique( ModulePath );
+				}
+			}
 
-		// Set the symbol path
-		Symbol->SetImagePathWide( *CombinedPath );
-		Symbol->SetSymbolPathWide( *CombinedPath );
-	
-		//@TODO: ROCKETHACK: This is for Rocket only.
-		Symbol->AppendImagePathWide( TEXT( "..\\..\\..\\Rocket\\Installed\\Windows\\Engine\\Binaries\\Win64" ) );
-		Symbol->AppendImagePathWide( TEXT( "..\\..\\..\\Rocket\\Installed\\Windows\\Engine\\Binaries\\Win32" ) );
-		Symbol->AppendSymbolPathWide( TEXT( "..\\..\\..\\Rocket\\Symbols\\Engine\\Binaries\\Win64" ) );
-		Symbol->AppendSymbolPathWide( TEXT( "..\\..\\..\\Rocket\\Symbols\\Engine\\Binaries\\Win32" ) );
-		Symbol->AppendImagePathWide( TEXT( "..\\..\\..\\Rocket\\LauncherInstalled\\Windows\\Launcher\\Engine\\Binaries\\Win64" ) );
-		Symbol->AppendSymbolPathWide( TEXT( "..\\..\\..\\Rocket\\LauncherSymbols\\Windows\\Launcher\\Engine\\Binaries\\Win64" ) );
+			for( int32 PathIndex = 0; PathIndex < SymbolPaths.Num(); PathIndex++ )
+			{
+				CombinedPath += FString( TEXT( "../../../" ) ) + SymbolPaths[PathIndex] + TEXT( ";" );
+			}
+
+			// Set the symbol path
+			Symbol->SetImagePathWide( *CombinedPath );
+			Symbol->SetSymbolPathWide( *CombinedPath );
+
+			//@TODO: ROCKETHACK: This is for Rocket only.
+			Symbol->AppendImagePathWide( TEXT( "..\\..\\..\\Rocket\\Installed\\Windows\\Engine\\Binaries\\Win64" ) );
+			Symbol->AppendImagePathWide( TEXT( "..\\..\\..\\Rocket\\Installed\\Windows\\Engine\\Binaries\\Win32" ) );
+			Symbol->AppendSymbolPathWide( TEXT( "..\\..\\..\\Rocket\\Symbols\\Engine\\Binaries\\Win64" ) );
+			Symbol->AppendSymbolPathWide( TEXT( "..\\..\\..\\Rocket\\Symbols\\Engine\\Binaries\\Win32" ) );
+			Symbol->AppendImagePathWide( TEXT( "..\\..\\..\\Rocket\\LauncherInstalled\\Windows\\Launcher\\Engine\\Binaries\\Win64" ) );
+			Symbol->AppendSymbolPathWide( TEXT( "..\\..\\..\\Rocket\\LauncherSymbols\\Windows\\Launcher\\Engine\\Binaries\\Win64" ) );
+		}	
 	}
 
 	// Add in syncing of the Microsoft symbol servers if requested
