CrashDebugHelper - Fixed not including .exe files in the SyncModules

diff --git a/Engine/Source/Developer/CrashDebugHelper/Private/CrashDebugHelper.cpp b/Engine/Source/Developer/CrashDebugHelper/Private/CrashDebugHelper.cpp
--- a/Engine/Source/Developer/CrashDebugHelper/Private/CrashDebugHelper.cpp
+++ b/Engine/Source/Developer/CrashDebugHelper/Private/CrashDebugHelper.cpp
@@ -171,17 +171,25 @@ bool ICrashDebugHelper::SyncModules( FCrashInfo* CrashInfo )
 
 					// Grab all dll and pdb files for the specified label.
 					TArray< TSharedRef<class ISourceControlRevision, ESPMode::ThreadSafe> > DLLSourceControlRevisions;
-					const FString DLLsPath = FString::Printf( TEXT( "%s/....dll*" ), *DepotName );
+					const FString DLLsPath = FString::Printf( TEXT( "%s/....dll" ), *DepotName );
 					Label->GetFileRevisions( DLLsPath, DLLSourceControlRevisions );
 
+					TArray< TSharedRef<class ISourceControlRevision, ESPMode::ThreadSafe> > EXESourceControlRevisions;
+					const FString EXEsPath = FString::Printf( TEXT( "%s/....exe" ), *DepotName );
+					Label->GetFileRevisions( EXEsPath, EXESourceControlRevisions );
+
 					TArray< TSharedRef<class ISourceControlRevision, ESPMode::ThreadSafe> > PDBSourceControlRevisions;
-					const FString PDBsPath = FString::Printf( TEXT( "%s/....pdb*" ), *DepotName );
+					const FString PDBsPath = FString::Printf( TEXT( "%s/....pdb" ), *DepotName );
 					Label->GetFileRevisions( PDBsPath, PDBSourceControlRevisions );		
 			
-					TSet<FString> DLLPaths;
+					TSet<FString> ModulesPaths;
 					for( const auto& DLLSrc : DLLSourceControlRevisions )
 					{
-						DLLPaths.Add( DLLSrc->GetFilename().Replace( *(DepotName + TEXT( "/" )), TEXT( "" ) ) );
+						ModulesPaths.Add( DLLSrc->GetFilename().Replace( *(DepotName + TEXT( "/" )), TEXT( "" ) ) );
+					}
+					for( const auto& EXESrc : EXESourceControlRevisions )
+					{
+						ModulesPaths.Add( EXESrc->GetFilename().Replace( *(DepotName + TEXT( "/" )), TEXT( "" ) ) );
 					}
 
 					TSet<FString> PDBPaths;
@@ -191,16 +199,16 @@ bool ICrashDebugHelper::SyncModules( FCrashInfo* CrashInfo )
 					}
 
 					// Iterate through all module and see if we have dll and pdb associated with the module, if so add it to the files to sync.
-					for( const auto& ModuleNameDLL : CrashInfo->ModuleNames )
+					for( const auto& ModuleName : CrashInfo->ModuleNames )
 					{
-						const FString ModuleNamePDB = ModuleNameDLL.Replace( TEXT( ".dll" ), TEXT( ".pdb" ) );
+						const FString ModuleNamePDB = ModuleName.Replace( TEXT( ".dll" ), TEXT( ".pdb" ) ).Replace( TEXT( ".exe" ), TEXT( ".pdb" ) );
 
-						for( const auto& DLLPath : DLLPaths )
+						for( const auto& ModulePath : ModulesPaths )
 						{
-							const bool bContainsDLL = DLLPath.Contains( ModuleNameDLL );
-							if( bContainsDLL )
+							const bool bContainsModule = ModulePath.Contains( ModuleName );
+							if( bContainsModule )
 							{
-								FilesToSync.Add( DLLPath );
+								FilesToSync.Add( ModulePath );
 							}
 						}
 
@@ -931,9 +939,8 @@ FPDBCacheEntryRef FPDBCache::CreateAndAddPDBCacheEntry( const FString& OriginalL
 	// CWD = F:\depot\UE4-Releases\4.1\Engine\Binaries\Win64
 	// RootDir = F:\depot\UE4-Releases\4.1\
 
-//#define DO_LOCAL_TESTING 1
 #if	DO_LOCAL_TESTING
-	const FString RootDir = TEXT( "U:/P4EPIC/UE4-Releases/4.1/" );
+	const FString RootDir = FString( TEXT( "U:/P4EPIC" ) ) / DepotName.Replace( TEXT( "//depot/" ), TEXT( "" ) );
 #else
 	const FString RootDir = FPaths::RootDir();
 #endif // DO_LOCAL_TESTING
@@ -955,12 +962,12 @@ FPDBCacheEntryRef FPDBCache::CreateAndAddPDBCacheEntry( const FString& OriginalL
 	}
 
 
-	TArray<FString> PDBFiles;
-	IFileManager::Get().FindFilesRecursive( PDBFiles, *EntryDirectory, TEXT( "*.*" ), true, false );
+	TArray<FString> CachedFiles;
+	IFileManager::Get().FindFilesRecursive( CachedFiles, *EntryDirectory, TEXT( "*.*" ), true, false );
 
 	// Calculate the size of this PDB Cache entry.
 	int64 TotalSize = 0;
-	for( const auto& Filename : PDBFiles )
+	for( const auto& Filename : CachedFiles )
 	{
 		const int64 FileSize = IFileManager::Get().FileSize( *Filename );
 		TotalSize += FileSize;
@@ -971,7 +978,7 @@ FPDBCacheEntryRef FPDBCache::CreateAndAddPDBCacheEntry( const FString& OriginalL
 
 	FPDBCacheEntryRef NewCacheEntry = MakeShareable( new FPDBCacheEntry( CleanedLabelName, SizeGB ) );
 	NewCacheEntry->LastAccessTime = LastAccessTime;
-	NewCacheEntry->Files = PDBFiles;
+	NewCacheEntry->Files = CachedFiles;
 	
 	PDBCacheEntries.Add( CleanedLabelName, NewCacheEntry );
 	SortPDBCache();
