CrashReport - Create PDB cache directory if needed

diff --git a/Engine/Source/Developer/CrashDebugHelper/Private/CrashDebugHelper.cpp b/Engine/Source/Developer/CrashDebugHelper/Private/CrashDebugHelper.cpp
--- a/Engine/Source/Developer/CrashDebugHelper/Private/CrashDebugHelper.cpp
+++ b/Engine/Source/Developer/CrashDebugHelper/Private/CrashDebugHelper.cpp
@@ -847,6 +847,7 @@ void FPDBCache::Init()
 void FPDBCache::InitializePDBCache()
 {
 	const double StartTime = FPlatformTime::Seconds();
+	IFileManager::Get().MakeDirectory( *PDBCachePath, true );
 
 	TArray<FString> PDBCacheEntryDirectories;
 	IFileManager::Get().FindFiles( PDBCacheEntryDirectories, *PDBCachePath, false, true );
