CrashReport - Misc fixed for local testing

diff --git a/Engine/Source/Developer/CrashDebugHelper/Private/Windows/WindowsPlatformStackWalkExt.cpp b/Engine/Source/Developer/CrashDebugHelper/Private/Windows/WindowsPlatformStackWalkExt.cpp
--- a/Engine/Source/Developer/CrashDebugHelper/Private/Windows/WindowsPlatformStackWalkExt.cpp
+++ b/Engine/Source/Developer/CrashDebugHelper/Private/Windows/WindowsPlatformStackWalkExt.cpp
@@ -212,7 +212,9 @@ void FWindowsPlatformStackWalkExt::SetSymbolPathsFromModules( FCrashInfo* CrashI
 		if( bUseNewSetPath )
 		{
 #if	DO_LOCAL_TESTING
-			const FString RootDir = TEXT( "U:/P4EPIC/UE4-Releases/4.1/" );
+			FString DepotName;
+			FParse::Value( FCommandLine::Get(), TEXT( "BranchName=" ), DepotName );
+			const FString RootDir = FString( TEXT( "U:/P4EPIC" ) ) / DepotName;
 #else
 			const FString RootDir = FPaths::RootDir();
 #endif // DO_LOCAL_TESTING
