Fix multi-line comment error.

(By removing leftover commented debug code).

diff --git a/Engine/Source/Developer/CrashDebugHelper/Private/CrashDebugHelper.cpp b/Engine/Source/Developer/CrashDebugHelper/Private/CrashDebugHelper.cpp
--- a/Engine/Source/Developer/CrashDebugHelper/Private/CrashDebugHelper.cpp
+++ b/Engine/Source/Developer/CrashDebugHelper/Private/CrashDebugHelper.cpp
@@ -932,14 +932,6 @@ FPDBCacheEntryRef FPDBCache::CreateAndAddPDBCacheEntry( const FString& OriginalL
 	UE_CLOG( !bSaved, LogCrashDebugHelper, Fatal, TEXT( "Couldn't save the timestamp file to %s" ), *EntryTimeStampFilename );
 	const FDateTime LastAccessTime = IFileManager::Get().GetTimeStamp( *EntryTimeStampFilename );
 
-	// Copy all synced files to the PDC Cache entry directory.
-	// DepotName = //depot/UE4-Releases/4.1/
-	// OriginalLabelName = //depot/UE4-Releases/4.2/Rocket-CL-2082666
-	// CleanedLabelName = __depot_UE4-Releases_4.2_Rocket-CL-2082666
-	// SyncedFiles = { //depot/UE4-Releases/4.1/engine/binaries/dotnet/swarminterface*.dll, ... }
-	// CWD = F:\depot\UE4-Releases\4.1\Engine\Binaries\Win64
-	// RootDir = F:\depot\UE4-Releases\4.1\
-
 #if	DO_LOCAL_TESTING
 	const FString RootDir = FString( TEXT( "U:/P4EPIC" ) ) / DepotName.Replace( TEXT( "//depot/" ), TEXT( "" ) );
 #else
