CrashDebugHelper - Commented obsolete functionality

diff --git a/Engine/Source/Developer/CrashDebugHelper/Private/Windows/CrashDebugHelperWindows.cpp b/Engine/Source/Developer/CrashDebugHelper/Private/Windows/CrashDebugHelperWindows.cpp
--- a/Engine/Source/Developer/CrashDebugHelper/Private/Windows/CrashDebugHelperWindows.cpp
+++ b/Engine/Source/Developer/CrashDebugHelper/Private/Windows/CrashDebugHelperWindows.cpp
@@ -331,7 +331,8 @@ bool FCrashDebugHelperWindows::ParseCrashDump_Windows(const TCHAR* InCrashDumpFi
 					UE_LOG(LogCrashDebugHelper, Display, TEXT("\tEngine Version = %d"), OutCrashDebugInfo.EngineVersion);
 					UE_LOG(LogCrashDebugHelper, Display, TEXT("\tPlatform       = %s"), *(OutCrashDebugInfo.PlatformName));
 
-					OutCrashDebugInfo.SourceControlLabel = GetLabelFromEngineVersion(OutCrashDebugInfo.EngineVersion);
+					// Obsolete?
+					//OutCrashDebugInfo.SourceControlLabel = GetLabelFromEngineVersion(OutCrashDebugInfo.EngineVersion);
 					if (OutCrashDebugInfo.SourceControlLabel.Len() > 0)
 					{
 						
