CrashDebugHelper
- Minidump diagnostics run from the crash processor will always use the PDB cache and the MDD from the main branch
- Reimplemented PDB cache to support UE4 releases that are stored on the network drive
- Cleaned up a bit code, removed old or obsolete functionality
- PDB Cache is disabled by default, can be only enabled through the command line (should fix issues related to "\\Device\\Harddisk3\\DR3", somehow caused by the bad flash device driver when F:\\ is a generic flash reader?)

diff --git a/Engine/Source/Developer/CrashDebugHelper/Private/CrashDebugHelper.cpp b/Engine/Source/Developer/CrashDebugHelper/Private/CrashDebugHelper.cpp
--- a/Engine/Source/Developer/CrashDebugHelper/Private/CrashDebugHelper.cpp
+++ b/Engine/Source/Developer/CrashDebugHelper/Private/CrashDebugHelper.cpp
@@ -13,6 +13,8 @@
 	#define MINIDUMPDIAGNOSTICS	0
 #endif
 
+static const TCHAR* P4_DEPOT_PREFIX = TEXT( "//depot/" );
+
 /** 
  * Global initialisation of this module
  */
@@ -25,7 +27,7 @@ bool ICrashDebugHelper::Init()
 	FString CmdLineBranchName;
 	if ( FParse::Value(FCommandLine::Get(), TEXT("BranchName="), CmdLineBranchName) )
 	{
-		DepotName = FString::Printf( TEXT( "//depot/%s" ), *CmdLineBranchName );
+		DepotName = FString::Printf( TEXT( "%s%s" ), P4_DEPOT_PREFIX, *CmdLineBranchName );
 	}
 	// Try to use what is configured in ini
 	else if (GConfig->GetString(TEXT("Engine.CrashDebugHelper"), TEXT("DepotName"), DepotName, GEngineIni) == true)
@@ -36,12 +38,14 @@ bool ICrashDebugHelper::Init()
 	// Default to BRANCH_NAME
 	else
 	{
-		DepotName = FString::Printf( TEXT( "//depot/%s" ), TEXT( BRANCH_NAME ) );
+		DepotName = FString::Printf( TEXT( "%s%s" ), P4_DEPOT_PREFIX, TEXT( BRANCH_NAME ) );
 	}
 
+	CrashInfo.DepotName = DepotName;
+
 	// Try to get the BuiltFromCL from command line to use this instead of attempting to locate the CL in the minidump
 	FString CmdLineBuiltFromCL;
-	BuiltFromCL = -1;
+	int32 BuiltFromCL = -1;
 	if (FParse::Value(FCommandLine::Get(), TEXT("BuiltFromCL="), CmdLineBuiltFromCL))
 	{
 		if ( !CmdLineBuiltFromCL.IsEmpty() )
@@ -49,13 +53,17 @@ bool ICrashDebugHelper::Init()
 			BuiltFromCL = FCString::Atoi(*CmdLineBuiltFromCL);
 		}
 	}
-
-	// Look up the source control search pattern. This pattern is used to identify a build label in the event that it was not already found in the database.
-	if (GConfig->GetString(TEXT("Engine.CrashDebugHelper"), TEXT("SourceControlBuildLabelPattern"), SourceControlBuildLabelPattern, GEngineIni) == true)
+	// Default to BUILT_FROM_CHANGELIST.
+	else
 	{
-		// Using a build label pattern to search in P4 if we don't find the label in the database
+		BuiltFromCL = int32(BUILT_FROM_CHANGELIST);
 	}
 
+	CrashInfo.ChangelistBuiltFrom = BuiltFromCL;
+
+	GConfig->GetString( TEXT( "Engine.CrashDebugHelper" ), TEXT( "SourceControlBuildLabelPattern" ), SourceControlBuildLabelPattern, GEngineIni );
+	GConfig->GetString( TEXT( "Engine.CrashDebugHelper" ), TEXT( "ExecutablePathPattern" ), ExecutablePathPattern, GEngineIni );
+
 	// Look up the local symbol store - fail if not found
 	if (GConfig->GetString(TEXT("Engine.CrashDebugHelper"), TEXT("LocalSymbolStore"), LocalSymbolStore, GEngineIni) == false)
 	{
@@ -63,7 +71,14 @@ bool ICrashDebugHelper::Init()
 		bInitialized = false;
 	}
 
-	PDBCache.Init();
+	if( !GConfig->GetString( TEXT( "Engine.CrashDebugHelper" ), TEXT( "DepotRoot" ), DepotRoot, GEngineIni ) )
+	{
+		UE_LOG( LogCrashDebugHelper, Warning, TEXT( "Failed to get DepotRoot from ini file... PDB Cache disabled" ) );
+	}
+	else
+	{
+		PDBCache.Init();
+	}
 
 	return bInitialized;
 }
@@ -117,9 +132,9 @@ void ICrashDebugHelper::ShutdownSourceControl()
 /** 
  * Sync the branch root relative file names to the requested label
  */
-bool ICrashDebugHelper::SyncModules( FCrashInfo* CrashInfo )
+bool ICrashDebugHelper::SyncModules()
 {
-	if( CrashInfo->LabelName.Len() <= 0 )
+	if( CrashInfo.LabelName.Len() <= 0 )
 	{
 		UE_LOG( LogCrashDebugHelper, Warning, TEXT( "SyncModules: Invalid Label parameter." ) );
 		return false;
@@ -131,43 +146,116 @@ bool ICrashDebugHelper::SyncModules( FCrashInfo* CrashInfo )
 		return false;
 	}
 
+	const TCHAR* UESymbols = TEXT( "Rocket/Symbols/" );
+	const bool bUseNetworkPathExecutables = !CrashInfo.NetworkPath.IsEmpty();
+
 	// Get all labels associated with the crash info's label.
-	// It should return one unique label because the crash info's label looks like this //depot/UE4-Releases/4.0/Rocket-CL-2047835 @see RetrieveBuildLabel
+	// It should return one unique label because the crash info's label looks like this //depot/UE4-Releases/4.0/UE-CL-2047835 @see RetrieveBuildLabel
 	// We can assume that the label name is unique.
 	// 
-	TArray< TSharedRef<ISourceControlLabel> > Labels = ISourceControlModule::Get().GetProvider().GetLabels( CrashInfo->LabelName );
+	TArray< TSharedRef<ISourceControlLabel> > Labels = ISourceControlModule::Get().GetProvider().GetLabels( CrashInfo.LabelName );
 	if( Labels.Num() == 1 )
 	{
-		const bool bContains = PDBCache.UsePDBCache() && PDBCache.ContainsPDBCacheEntry( CrashInfo->LabelName );
-		if( bContains )
-		{
-			UE_LOG( LogCrashDebugHelper, Warning, TEXT( "Label %s found in the PDB Cache, using it" ), *CrashInfo->LabelName );
-			CrashInfo->PDBCacheEntry = PDBCache.FindPDBCacheEntry( CrashInfo->LabelName );
-			PDBCache.TouchPDBCacheEntry( CrashInfo->LabelName );
-			return true;
-		}
-		else
-		{
-			TArray<FString> SyncedFiles;	
-			TSharedRef<ISourceControlLabel> Label = Labels[0];
+		TSharedRef<ISourceControlLabel> Label = Labels[0];
+		TSet<FString> FilesToSync;
 
-			//@TODO: MAC: Excluding labels for Mac since we are only syncing windows binaries here...
-			if( Label->GetName().Contains( TEXT( "Mac" ) ) )
+		if( bUseNetworkPathExecutables )
+		{
+			// Use product version instead of label name to make a distinguish between chosen methods.
+			const bool bContains = PDBCache.UsePDBCache() && PDBCache.ContainsPDBCacheEntry( CrashInfo.ProductVersion );
+			if( bContains )
 			{
-				UE_LOG( LogCrashDebugHelper, Log, TEXT( " Skipping Mac label '%s' when syncing modules." ), *Label->GetName() );
+				UE_LOG( LogCrashDebugHelper, Warning, TEXT( "Label %s found in the PDB Cache, using it" ), *CrashInfo.ProductVersion );
+				CrashInfo.PDBCacheEntry = PDBCache.FindAndTouchPDBCacheEntry( CrashInfo.ProductVersion );
+				return true;
 			}
 			else
 			{
-				// Sync all the dll, exes, and related symbol files
-				UE_LOG( LogCrashDebugHelper, Log, TEXT( " Syncing modules with label '%s'." ), *Label->GetName() );
+				SCOPE_LOG_TIME_IN_SECONDS( TEXT( "SyncModulesAndNetwork" ), nullptr );
+
+				// This is a bit hacky, probably will be valid until next build system change.
+				const FString ProductNetworkPath = ExecutablePathPattern.Replace( TEXT( "%PRODUCT_VERSION%" ), *CrashInfo.ProductVersion );
+
+				// Grab information about symbols.
+				TArray< TSharedRef<class ISourceControlRevision, ESPMode::ThreadSafe> > PDBSourceControlRevisions;
+				const FString PDBsPath = FString::Printf( TEXT( "%s/%s....pdb" ), *DepotName, UESymbols );
+				Label->GetFileRevisions( PDBsPath, PDBSourceControlRevisions );
+				
+				TSet<FString> PDBPaths;
+				for( const auto& PDBSrc : PDBSourceControlRevisions )
+				{
+					PDBPaths.Add( PDBSrc->GetFilename() );
+				}
 
-				TSet<FString> FilesToSync;
+				// Now, sync symbols.
+				for( const auto& PDBPath : PDBPaths )
+				{
+					if( Label->Sync( PDBPath ) )
+					{
+						UE_LOG( LogCrashDebugHelper, Warning, TEXT( "Synced PDB '%s'" ), *PDBPath );
+					}
+				}
+
+				// Find all the executables in the product network path.
+				TArray<FString> NetworkExecutables;
+				IFileManager::Get().FindFilesRecursive( NetworkExecutables, *ProductNetworkPath, TEXT( "*.dll" ), true, false, false );
+				IFileManager::Get().FindFilesRecursive( NetworkExecutables, *ProductNetworkPath, TEXT( "*.exe" ), true, false, false );
 
-				const bool bUseNewSyncModule = true;
+				// From=Full pathname
+				// To=Relative pathname
+				TMap<FString, FString> FilesToBeCached;
 
-				if( bUseNewSyncModule )
+				// If a symbol matches an executable, add the pair to the list of files that should be cached.
+				for( const auto& NetworkExecutableFullpath : NetworkExecutables )
 				{
-					SCOPE_LOG_TIME_IN_SECONDS( TEXT( "SyncModules_NEW" ), nullptr );
+					for( const auto& PDBPath : PDBPaths )
+					{
+						const FString PDBRelativePath = PDBPath.Replace( *DepotName, TEXT( "" ) ).Replace( UESymbols, TEXT( "" ) );
+						const FString PDBFullpath = DepotRoot / PDBPath.Replace( P4_DEPOT_PREFIX, TEXT( "" ) );
+
+						const FString PDBMatch = PDBRelativePath.Replace( TEXT( "pdb" ), TEXT( "" ) );
+						const FString NetworkRelativePath = NetworkExecutableFullpath.Replace( *ProductNetworkPath, TEXT( "" ) );
+						const bool bMatch = NetworkExecutableFullpath.Contains( PDBMatch );
+						if( bMatch )
+						{
+							// From -> Where ?
+							FilesToBeCached.Add( NetworkExecutableFullpath, NetworkRelativePath );
+							FilesToBeCached.Add( PDBFullpath, PDBRelativePath );
+							break;
+						}
+					}
+				}
+
+				if( PDBCache.UsePDBCache() )
+				{
+					// Initialize and add a new PDB Cache entry to the database.
+					CrashInfo.PDBCacheEntry = PDBCache.CreateAndAddPDBCacheEntryMixed( CrashInfo.ProductVersion, FilesToBeCached );
+				}
+			}
+		}
+		else
+		{
+			TArray<FString> FilesToBeCached;
+			const bool bContains = PDBCache.UsePDBCache() && PDBCache.ContainsPDBCacheEntry( CrashInfo.LabelName );
+			if( bContains )
+			{
+				UE_LOG( LogCrashDebugHelper, Warning, TEXT( "Label '%s' found in the PDB Cache, using it" ), *CrashInfo.LabelName );
+				CrashInfo.PDBCacheEntry = PDBCache.FindAndTouchPDBCacheEntry( CrashInfo.LabelName );
+				return true;
+			}
+			else
+			{
+				//@TODO: MAC: Excluding labels for Mac since we are only syncing windows binaries here...
+				if( Label->GetName().Contains( TEXT( "Mac" ) ) )
+				{
+					UE_LOG( LogCrashDebugHelper, Log, TEXT( " Skipping Mac label '%s' when syncing modules." ), *Label->GetName() );
+				}
+				else
+				{
+					// Sync all the dll, exes, and related symbol files
+					UE_LOG( LogCrashDebugHelper, Log, TEXT( " Syncing modules with label '%s'." ), *Label->GetName() );
+
+					SCOPE_LOG_TIME_IN_SECONDS( TEXT( "SyncModules" ), nullptr );
 
 					// Grab all dll and pdb files for the specified label.
 					TArray< TSharedRef<class ISourceControlRevision, ESPMode::ThreadSafe> > DLLSourceControlRevisions;
@@ -180,26 +268,26 @@ bool ICrashDebugHelper::SyncModules( FCrashInfo* CrashInfo )
 
 					TArray< TSharedRef<class ISourceControlRevision, ESPMode::ThreadSafe> > PDBSourceControlRevisions;
 					const FString PDBsPath = FString::Printf( TEXT( "%s/....pdb" ), *DepotName );
-					Label->GetFileRevisions( PDBsPath, PDBSourceControlRevisions );		
-			
+					Label->GetFileRevisions( PDBsPath, PDBSourceControlRevisions );
+
 					TSet<FString> ModulesPaths;
 					for( const auto& DLLSrc : DLLSourceControlRevisions )
 					{
-						ModulesPaths.Add( DLLSrc->GetFilename().Replace( *(DepotName + TEXT( "/" )), TEXT( "" ) ) );
+						ModulesPaths.Add( DLLSrc->GetFilename().Replace( *DepotName, TEXT( "" ) ) );
 					}
 					for( const auto& EXESrc : EXESourceControlRevisions )
 					{
-						ModulesPaths.Add( EXESrc->GetFilename().Replace( *(DepotName + TEXT( "/" )), TEXT( "" ) ) );
+						ModulesPaths.Add( EXESrc->GetFilename().Replace( *DepotName, TEXT( "" ) ) );
 					}
 
 					TSet<FString> PDBPaths;
 					for( const auto& PDBSrc : PDBSourceControlRevisions )
 					{
-						PDBPaths.Add( PDBSrc->GetFilename().Replace( *(DepotName + TEXT( "/" )), TEXT( "" ) ) );
+						PDBPaths.Add( PDBSrc->GetFilename().Replace( *DepotName, TEXT( "" ) ) );
 					}
 
 					// Iterate through all module and see if we have dll and pdb associated with the module, if so add it to the files to sync.
-					for( const auto& ModuleName : CrashInfo->ModuleNames )
+					for( const auto& ModuleName : CrashInfo.ModuleNames )
 					{
 						const FString ModuleNamePDB = ModuleName.Replace( TEXT( ".dll" ), TEXT( ".pdb" ) ).Replace( TEXT( ".exe" ), TEXT( ".pdb" ) );
 
@@ -221,91 +309,30 @@ bool ICrashDebugHelper::SyncModules( FCrashInfo* CrashInfo )
 							}
 						}
 					}
-					
+
 					// Now, sync all files.
 					for( const auto& Filename : FilesToSync )
 					{
-						const FString DepotPath = FString::Printf( TEXT( "%s/%s" ), *DepotName, *Filename );
+						const FString DepotPath = DepotName + Filename;
 						if( Label->Sync( DepotPath ) )
 						{
 							UE_LOG( LogCrashDebugHelper, Warning, TEXT( " ... synced binary '%s'." ), *DepotPath );
 						}
-						SyncedFiles.Add( DepotPath );
+						FilesToBeCached.Add( DepotPath );
 					}
 				}
-				else
-				{
-					SCOPE_LOG_TIME_IN_SECONDS( TEXT( "SyncModules_OLD" ), nullptr );
-
-					for( int32 ModuleNameIndex = 0; ModuleNameIndex < CrashInfo->ModuleNames.Num(); ModuleNameIndex++ )
-					{
-						// Sometimes P4 may hang due to spamming the server with the commands, if so try uncommenting the line below.
-						//FPlatformProcess::Sleep( 0.1f );
-						FString DepotPath = FString::Printf( TEXT( "%s/%s" ), *DepotName, *CrashInfo->ModuleNames[ModuleNameIndex] );
-						// Match all decorated versions of the module. We may need them if the file was renamed to remove the "-Platform-Configuration" decoration.
-						DepotPath = DepotPath.Replace( TEXT( ".dll" ), TEXT( "*.dll" ) ).Replace( TEXT( ".exe" ), TEXT( "*.exe" ) );
 
-						if( Label->Sync( DepotPath ) )
-						{
-							UE_LOG( LogCrashDebugHelper, Warning, TEXT( " ... synced binary '%s'." ), *DepotPath );
-						}
-						SyncedFiles.Add( DepotPath );
-
-						FString PDBName = DepotPath.Replace( TEXT( ".dll" ), TEXT( ".pdb" ) ).Replace( TEXT( ".exe" ), TEXT( ".pdb" ) );
-						if( Label->Sync( PDBName ) )
-						{
-							UE_LOG( LogCrashDebugHelper, Warning, TEXT( " ... synced symbol '%s'." ), *PDBName );
-						}
-						SyncedFiles.Add( PDBName );
-
-						//@TODO: ROCKETHACK: Adding additional Installed and Symbol paths - revisit when builds are made by the builder...
-						DepotPath = FString::Printf( TEXT( "%s/Rocket/Installed/Windows/%s" ), *DepotName, *CrashInfo->ModuleNames[ModuleNameIndex] );
-						// Match all decorated versions of the module. We may need them if the file was renamed to remove the "-Platform-Configuration" decoration.
-						DepotPath = DepotPath.Replace( TEXT( ".dll" ), TEXT( "*.dll" ) ).Replace( TEXT( ".exe" ), TEXT( "*.exe" ) );
-						if( Label->Sync( DepotPath ) )
-						{
-							UE_LOG( LogCrashDebugHelper, Warning, TEXT( " ... synced binary '%s'." ), *DepotPath );
-						}
-						SyncedFiles.Add( DepotPath );
-
-						DepotPath = FString::Printf( TEXT( "%s/Rocket/Symbols/%s" ), *DepotName, *CrashInfo->ModuleNames[ModuleNameIndex] );
-						PDBName = DepotPath.Replace( TEXT( ".dll" ), TEXT( "*.pdb" ) ).Replace( TEXT( ".exe" ), TEXT( "*.pdb" ) );
-						if( Label->Sync( PDBName ) )
-						{
-							UE_LOG( LogCrashDebugHelper, Warning, TEXT( " ... synced symbol '%s'." ), *PDBName );
-						}
-						SyncedFiles.Add( PDBName );
-
-						DepotPath = FString::Printf( TEXT( "%s/Rocket/LauncherInstalled/Windows/Launcher/%s" ), *DepotName, *CrashInfo->ModuleNames[ModuleNameIndex] );
-						// Match all decorated versions of the module. We may need them if the file was renamed to remove the "-Platform-Configuration" decoration.
-						DepotPath = DepotPath.Replace( TEXT( ".dll" ), TEXT( "*.dll" ) ).Replace( TEXT( ".exe" ), TEXT( "*.exe" ) );
-						if( Label->Sync( DepotPath ) )
-						{
-							UE_LOG( LogCrashDebugHelper, Warning, TEXT( " ... synced binary '%s'." ), *DepotPath );
-						}
-						SyncedFiles.Add( DepotPath );
-
-						DepotPath = FString::Printf( TEXT( "%s/Rocket/LauncherSymbols/Windows/Launcher/%s" ), *DepotName, *CrashInfo->ModuleNames[ModuleNameIndex] );
-						PDBName = DepotPath.Replace( TEXT( ".dll" ), TEXT( "*.pdb" ) ).Replace( TEXT( ".exe" ), TEXT( "*.pdb" ) );
-						if( Label->Sync( PDBName ) )
-						{
-							UE_LOG( LogCrashDebugHelper, Warning, TEXT( " ... synced symbol '%s'." ), *PDBName );
-						}
-						SyncedFiles.Add( PDBName );
-					}
+				if( PDBCache.UsePDBCache() )
+				{
+					// Initialize and add a new PDB Cache entry to the database.
+					CrashInfo.PDBCacheEntry = PDBCache.CreateAndAddPDBCacheEntry( CrashInfo.LabelName, DepotRoot, DepotName, FilesToBeCached );
 				}
 			}
-			
-			if( PDBCache.UsePDBCache() )
-			{
-				// Initialize and add a new PDB Cache entry to the database.
-				CrashInfo->PDBCacheEntry = PDBCache.CreateAndAddPDBCacheEntry( CrashInfo->LabelName, DepotName, SyncedFiles );
-			}
-		}		
+		}
 	}
 	else
 	{
-		UE_LOG( LogCrashDebugHelper, Error, TEXT( "Could not find label '%s'."), *CrashInfo->LabelName );
+		UE_LOG( LogCrashDebugHelper, Error, TEXT( "Could not find label '%s'."), *CrashInfo.LabelName );
 	}
 
 	return true;
@@ -314,9 +341,9 @@ bool ICrashDebugHelper::SyncModules( FCrashInfo* CrashInfo )
 /** 
  * Sync a single source file to the requested label
  */
-bool ICrashDebugHelper::SyncSourceFile( FCrashInfo* CrashInfo )
+bool ICrashDebugHelper::SyncSourceFile()
 {
-	if (CrashInfo->LabelName.Len() <= 0)
+	if (CrashInfo.LabelName.Len() <= 0)
 	{
 		UE_LOG( LogCrashDebugHelper, Warning, TEXT( "SyncSourceFile: Invalid Label parameter." ) );
 		return false;
@@ -329,8 +356,8 @@ bool ICrashDebugHelper::SyncSourceFile( FCrashInfo* CrashInfo )
 	}
 
 	// Sync all the dll, exes, and related symbol files
-	FString DepotPath = FString::Printf( TEXT( "%s/%s" ), *DepotName, *CrashInfo->SourceFile );
-	TArray< TSharedRef<ISourceControlLabel> > Labels = ISourceControlModule::Get().GetProvider().GetLabels( CrashInfo->LabelName );
+	FString DepotPath = DepotName / CrashInfo.SourceFile;
+	TArray< TSharedRef<ISourceControlLabel> > Labels = ISourceControlModule::Get().GetProvider().GetLabels( CrashInfo.LabelName );
 	if(Labels.Num() > 0)
 	{
 		TSharedRef<ISourceControlLabel> Label = Labels[0];
@@ -341,7 +368,7 @@ bool ICrashDebugHelper::SyncSourceFile( FCrashInfo* CrashInfo )
 	}
 	else
 	{
-		UE_LOG( LogCrashDebugHelper, Error, TEXT( "Could not find label '%s'."), *CrashInfo->LabelName );
+		UE_LOG( LogCrashDebugHelper, Error, TEXT( "Could not find label '%s'."), *CrashInfo.LabelName );
 	}
 
 	return true;
@@ -376,26 +403,26 @@ bool ICrashDebugHelper::ReadSourceFile( const TCHAR* InFilename, TArray<FString>
 /** 
  * Add adjacent lines of the source file the crash occurred in the crash report
  */
-void ICrashDebugHelper::AddSourceToReport( FCrashInfo* CrashInfo )
+void ICrashDebugHelper::AddSourceToReport()
 {
-	if( CrashInfo->SourceFile.Len() > 0 && CrashInfo->SourceLineNumber != 0 )
+	if( CrashInfo.SourceFile.Len() > 0 && CrashInfo.SourceLineNumber != 0 )
 	{
 		TArray<FString> Lines;
-		FString FullPath = FString( TEXT( "../../../" ) ) + CrashInfo->SourceFile;
+		FString FullPath = FString( TEXT( "../../../" ) ) + CrashInfo.SourceFile;
 		ReadSourceFile( *FullPath, Lines );
 
-		uint64 MinLine = FMath::Clamp( CrashInfo->SourceLineNumber - 15, ( uint64 )1, ( uint64 )Lines.Num() );
-		uint64 MaxLine = FMath::Clamp( CrashInfo->SourceLineNumber + 15, ( uint64 )1, ( uint64 )Lines.Num() );
+		uint64 MinLine = FMath::Clamp( CrashInfo.SourceLineNumber - 15, ( uint64 )1, ( uint64 )Lines.Num() );
+		uint64 MaxLine = FMath::Clamp( CrashInfo.SourceLineNumber + 15, ( uint64 )1, ( uint64 )Lines.Num() );
 
 		for( int32 Line = MinLine; Line < MaxLine; Line++ )
 		{
-			if( Line == CrashInfo->SourceLineNumber - 1 )
+			if( Line == CrashInfo.SourceLineNumber - 1 )
 			{
-				CrashInfo->SourceContext.Add( FString( TEXT( "*****" ) ) + Lines[Line] );
+				CrashInfo.SourceContext.Add( FString( TEXT( "*****" ) ) + Lines[Line] );
 			}
 			else
 			{
-				CrashInfo->SourceContext.Add( FString( TEXT( "     " ) ) + Lines[Line] );
+				CrashInfo.SourceContext.Add( FString( TEXT( "     " ) ) + Lines[Line] );
 			}
 		}
 	}
@@ -404,10 +431,10 @@ void ICrashDebugHelper::AddSourceToReport( FCrashInfo* CrashInfo )
 /** 
  * Add source control annotated adjacent lines of the source file the crash occurred in the crash report
  */
-bool ICrashDebugHelper::AddAnnotatedSourceToReport( FCrashInfo* CrashInfo )
+bool ICrashDebugHelper::AddAnnotatedSourceToReport()
 {
 	// Make sure we have a source file to interrogate
-	if( CrashInfo->SourceFile.Len() > 0 && CrashInfo->SourceLineNumber != 0 )
+	if( CrashInfo.SourceFile.Len() > 0 && CrashInfo.SourceLineNumber != 0 )
 	{
 		// Check source control
 		if( !ISourceControlModule::Get().IsEnabled() )
@@ -416,24 +443,24 @@ bool ICrashDebugHelper::AddAnnotatedSourceToReport( FCrashInfo* CrashInfo )
 		}
 
 		// Ask source control to annotate the file for us
-		FString DepotPath = FString::Printf( TEXT( "%s/%s" ), *DepotName, *CrashInfo->SourceFile );
+		FString DepotPath = DepotName / CrashInfo.SourceFile;
 
 		TArray<FAnnotationLine> Lines;
-		SourceControlHelpers::AnnotateFile( ISourceControlModule::Get().GetProvider(), CrashInfo->LabelName, DepotPath, Lines );
+		SourceControlHelpers::AnnotateFile( ISourceControlModule::Get().GetProvider(), CrashInfo.LabelName, DepotPath, Lines );
 
-		uint64 MinLine = FMath::Clamp( CrashInfo->SourceLineNumber - 15, ( uint64 )1, ( uint64 )Lines.Num() );
-		uint64 MaxLine = FMath::Clamp( CrashInfo->SourceLineNumber + 15, ( uint64 )1, ( uint64 )Lines.Num() );
+		uint64 MinLine = FMath::Clamp( CrashInfo.SourceLineNumber - 15, ( uint64 )1, ( uint64 )Lines.Num() );
+		uint64 MaxLine = FMath::Clamp( CrashInfo.SourceLineNumber + 15, ( uint64 )1, ( uint64 )Lines.Num() );
 
 		// Display a source context in the report, and decorate each line with the last editor of the line
 		for( int32 Line = MinLine; Line < MaxLine; Line++ )
 		{			
-			if( Line == CrashInfo->SourceLineNumber )
+			if( Line == CrashInfo.SourceLineNumber )
 			{
-				CrashInfo->SourceContext.Add( FString::Printf( TEXT( "*****%20s: %s" ), *Lines[Line].UserName, *Lines[Line].Line ) );
+				CrashInfo.SourceContext.Add( FString::Printf( TEXT( "*****%20s: %s" ), *Lines[Line].UserName, *Lines[Line].Line ) );
 			}
 			else
 			{
-				CrashInfo->SourceContext.Add( FString::Printf( TEXT( "     %20s: %s" ), *Lines[Line].UserName, *Lines[Line].Line ) );
+				CrashInfo.SourceContext.Add( FString::Printf( TEXT( "     %20s: %s" ), *Lines[Line].UserName, *Lines[Line].Line ) );
 			}
 		}
 		return true;
@@ -514,7 +541,7 @@ void FCrashInfo::GenerateReport( const FString& DiagnosticsPath )
 
 		if( Modules.Num() > 0 )
 		{
-			Line = FString::Printf( TEXT( "Application version %d.%d.%d.%d" ), Modules[0].Major, Modules[0].Minor, Modules[0].Build, Modules[0].Revision );
+			Line = FString::Printf( TEXT( "Application version %d.%d.%d" ), Modules[0].Major, Modules[0].Minor, Modules[0].Patch );
 			WriteLine( ReportFile, TCHAR_TO_UTF8( *Line ) );
 		}
 
@@ -576,7 +603,7 @@ void FCrashInfo::GenerateReport( const FString& DiagnosticsPath )
 			FString ModuleName = FPaths::GetBaseFilename( Module.Name, true ) + FPaths::GetExtension( Module.Name, true );
 
 			FString ModuleDetail = FString::Printf( TEXT( "%40s" ), *ModuleName );
-			FString Version = FString::Printf( TEXT( " (%d.%d.%d.%d)" ), Module.Major, Module.Minor, Module.Build, Module.Revision );
+			FString Version = FString::Printf( TEXT( " (%d.%d.%d.%d)" ), Module.Major, Module.Minor, Module.Patch, Module.Revision );
 			ModuleDetail += FString::Printf( TEXT( " %22s" ), *Version );
 			ModuleDetail += FString::Printf( TEXT( " 0x%016x 0x%08x" ), Module.BaseOfImage, Module.SizeOfImage );
 			ModuleDetail += FString::Printf( TEXT( " %s" ), *ModuleDirectory );
@@ -675,7 +702,7 @@ bool ICrashDebugHelper::SyncRequiredFilesForDebuggingFromLabel(const FString& In
 
 				// Copy the files to a flat directory.
 				// This will have problems if there are any files named the same!!!!
-				FString LocalStoreFolder = FString::Printf(TEXT("%s/%s/"), *LocalSymbolStore, *InPlatform);
+				FString LocalStoreFolder = LocalSymbolStore / InPlatform;
 				FString CopyFilename = LocalStoreFolder / FPaths::GetCleanFilename(Revision->GetFilename());
 
 //				UE_LOG(LogCrashDebugHelper, Display, TEXT("Copying engine file: %s to %s"), *Filename, *CopyFilename);
@@ -704,7 +731,7 @@ bool ICrashDebugHelper::SyncRequiredFilesForDebuggingFromChangelist(int32 InChan
 	// Not really useful as the source indexed pdbs will be tied to labeled builds.
 	// For now, we will not support this
 
-	FString BuildLabel = RetrieveBuildLabel(-1, InChangelistNumber);
+	const FString BuildLabel = RetrieveBuildLabel(InChangelistNumber);
 	if (BuildLabel.Len() > 0)
 	{
 		return SyncRequiredFilesForDebuggingFromLabel(BuildLabel, InPlatform);
@@ -718,21 +745,20 @@ bool ICrashDebugHelper::SyncRequiredFilesForDebuggingFromChangelist(int32 InChan
 /**
  *	Retrieve the build label for the given engine version or changelist number.
  *
- *	@param	InEngineVersion		The engine version to retrieve the label for. If -1, then will use InChangelistNumber
  *	@param	InChangelistNumber	The changelist number to retrieve the label for
  *	@return	FString				The label if successful, empty if it wasn't found or another error
  */
-FString ICrashDebugHelper::RetrieveBuildLabel(int32 InEngineVersion, int32 InChangelistNumber)
+FString ICrashDebugHelper::RetrieveBuildLabel(int32 InChangelistNumber)
 {
 	FString FoundLabelString;
 
-	if ((InEngineVersion < 0) && (InChangelistNumber < 0))
+	if( InChangelistNumber < 0 )
 	{
 		UE_LOG(LogCrashDebugHelper, Warning, TEXT("RetrieveBuildLabel: Invalid parameters."));
 		return FoundLabelString;
 	}
 
-	// If we failed to find the label, try to find the label directly in source control by using the pattern supplied via ini
+	// Try to find the label directly in source control by using the pattern supplied via ini
 	if ( FoundLabelString.IsEmpty() && InChangelistNumber >= 0 && !SourceControlBuildLabelPattern.IsEmpty() )
 	{
 		const FString ChangelistString = FString::Printf(TEXT("%d"), InChangelistNumber);
@@ -754,6 +780,64 @@ FString ICrashDebugHelper::RetrieveBuildLabel(int32 InEngineVersion, int32 InCha
 	return FoundLabelString;
 }
 
+void ICrashDebugHelper::RetrieveBuildLabelAndNetworkPath( int32 InChangelistNumber )
+{
+	if( InChangelistNumber < 0 )
+	{
+		UE_LOG( LogCrashDebugHelper, Warning, TEXT( "RetrieveBuildLabelAndNetworkPath: Invalid parameters." ) );
+		CrashInfo.LabelName.Empty();
+		CrashInfo.NetworkPath.Empty();
+	}
+
+	// Try to find the label directly in source control by using the pattern supplied via ini.
+	if( InChangelistNumber >= 0 && !SourceControlBuildLabelPattern.IsEmpty() )
+	{
+		const FString ChangelistString = FString::Printf( TEXT( "%d" ), InChangelistNumber );
+		const FString TestLabelWithCL = SourceControlBuildLabelPattern.Replace( TEXT( "%CHANGELISTNUMBER%" ), *ChangelistString, ESearchCase::CaseSensitive );
+		TArray< TSharedRef<ISourceControlLabel> > Labels = ISourceControlModule::Get().GetProvider().GetLabels( TestLabelWithCL );
+		if( Labels.Num() > 0 )
+		{
+			// If we found more than one label, warn about it and just use the first one
+			if( Labels.Num() > 1 )
+			{
+				UE_LOG( LogCrashDebugHelper, Warning, TEXT( "RetrieveBuildLabelAndNetworkPath: More than one build label found with pattern '%s' - Using label '%s'" ), *TestLabelWithCL, *Labels[0]->GetName() );
+			}
+
+			CrashInfo.LabelName = Labels[0]->GetName();
+			UE_LOG( LogCrashDebugHelper, Log, TEXT( "RetrieveBuildLabelAndNetworkPath: Found label '%s' matching pattern '%s' in source control." ), *CrashInfo.LabelName, *TestLabelWithCL );
+		}
+	}
+	else
+	{
+		CrashInfo.LabelName.Empty();
+	}
+
+	// Try to find the network path by using the pattern supplied via ini.
+	// If this step successes, we will grab the executable from the network path instead of P4.
+	// At this moment we only care about UE4 releases.
+	const bool bCanUseNetworkPath = DepotName.Contains( TEXT( "UE4-Releases" ) );
+	if( bCanUseNetworkPath && InChangelistNumber >= 0 && !ExecutablePathPattern.IsEmpty() )
+	{
+		const FString TestNetworkPath = ExecutablePathPattern.Replace( TEXT( "%PRODUCT_VERSION%" ), *CrashInfo.ProductVersion );
+		const bool bIsValidPath = IFileManager::Get().DirectoryExists( *TestNetworkPath );
+
+		if( bIsValidPath )
+		{
+			CrashInfo.NetworkPath = TestNetworkPath;
+			UE_LOG( LogCrashDebugHelper, Log, TEXT( "RetrieveBuildLabelAndNetworkPath: Found network path '%s'." ), *CrashInfo.NetworkPath );
+		}
+		else
+		{
+			CrashInfo.NetworkPath.Empty();
+			UE_LOG( LogCrashDebugHelper, Log, TEXT( "RetrieveBuildLabelAndNetworkPath: Network path '%s' not found." ), *TestNetworkPath );
+		}
+	}
+	else
+	{
+		CrashInfo.NetworkPath.Empty();
+	}
+}
+
 /*-----------------------------------------------------------------------------
 	PDB Cache implementation
 -----------------------------------------------------------------------------*/
@@ -764,14 +848,14 @@ void FPDBCache::Init()
 {
 	// PDB Cache
 	// Default configuration
-	// PDBCachePath=U:/CrashReport/PDBCache
+	// PDBCachePath=F:/CrashReportPDBCache
 	// DaysToDeleteUnusedFilesFromPDBCache=14
 	// PDBCacheSizeGB=128
 	// MinFreeSizeGB=64
-	// bUsePDBCache=true
 
-	// Look up whether we want to use the PDB cache.
-	GConfig->GetBool( TEXT( "Engine.CrashDebugHelper" ), TEXT( "bUsePDBCache" ), bUsePDBCache, GEngineIni );
+	// Can be enabled only through the command line.
+	FParse::Bool( FCommandLine::Get(), TEXT( "bUsePDBCache=" ), bUsePDBCache );
+
 	UE_LOG( LogCrashDebugHelper, Warning, TEXT( "bUsePDBCache is %s" ), bUsePDBCache ? TEXT( "enabled" ) : TEXT( "disabled" ) );
 
 	// Get the rest of the PDB cache configuration.
@@ -852,10 +936,10 @@ void FPDBCache::InitializePDBCache()
 	TArray<FString> PDBCacheEntryDirectories;
 	IFileManager::Get().FindFiles( PDBCacheEntryDirectories, *PDBCachePath, false, true );
 
-	for( const auto& DirPath : PDBCacheEntryDirectories )
+	for( const auto& Directory : PDBCacheEntryDirectories )
 	{
-		FPDBCacheEntryRef Entry = ReadPDBCacheEntry( DirPath );
-		PDBCacheEntries.Add( DirPath, Entry );
+		FPDBCacheEntryRef Entry = ReadPDBCacheEntry( Directory );
+		PDBCacheEntries.Add( Directory, Entry );
 	}
 
 	SortPDBCache();
@@ -879,13 +963,13 @@ void FPDBCache::CleanPDBCache( int32 DaysToDelete, int32 NumberOfGBsToBeCleaned
 	for( const auto& It : PDBCacheEntries )
 	{
 		const FPDBCacheEntryRef& Entry = It.Value;
-		const FString EntryDirectory = PDBCachePath / Entry->Label;
+		const FString EntryDirectory = PDBCachePath / Entry->Directory;
 		const FString EntryTimeStampFilename = EntryDirectory / PDBTimeStampFile;
 
 		const double EntryFileAge = IFileManager::Get().GetFileAgeSeconds( *EntryTimeStampFilename );
 		if( EntryFileAge > DaysToDeleteAsSeconds )
 		{
-			EntriesToBeRemoved.Add( Entry->Label );
+			EntriesToBeRemoved.Add( Entry->Directory );
 			NumGBsCleaned += Entry->SizeGB;
 		}
 	}
@@ -897,9 +981,9 @@ void FPDBCache::CleanPDBCache( int32 DaysToDelete, int32 NumberOfGBsToBeCleaned
 		{
 			const FPDBCacheEntryRef& Entry = It.Value;
 
-			if( !EntriesToBeRemoved.Contains(Entry->Label) )
+			if( !EntriesToBeRemoved.Contains( Entry->Directory ) )
 			{
-				EntriesToBeRemoved.Add( Entry->Label );
+				EntriesToBeRemoved.Add( Entry->Directory );
 				NumGBsCleaned += Entry->SizeGB;
 
 				if( NumGBsCleaned > NumberOfGBsToBeCleaned )
@@ -912,36 +996,32 @@ void FPDBCache::CleanPDBCache( int32 DaysToDelete, int32 NumberOfGBsToBeCleaned
 	}
 
 	// Remove all marked PDB Cache entries.
-	for( const auto& EntryLabel : EntriesToBeRemoved )
+	for( const auto& EntryDirectory : EntriesToBeRemoved )
 	{
-		RemovePDBCacheEntry( EntryLabel );
+		RemovePDBCacheEntry( EntryDirectory );
 	}
 
 	const double TotalTime = FPlatformTime::Seconds() - StartTime;
 	UE_LOG( LogCrashDebugHelper, Log, TEXT( "PDB Cache cleaned in %.2f ms" ), TotalTime*1000.0f );
 }
 
-FPDBCacheEntryRef FPDBCache::CreateAndAddPDBCacheEntry( const FString& OriginalLabelName, const FString& DepotName, const TArray<FString>& SyncedFiles )
+FPDBCacheEntryRef FPDBCache::CreateAndAddPDBCacheEntry( const FString& OriginalLabelName, const FString& DepotRoot, const FString& DepotName, const TArray<FString>& FilesToBeCached )
 {
-	const FString CleanedLabelName = CleanLabelName( OriginalLabelName );
+	const FString CleanedLabelName = EscapePath( OriginalLabelName );
 	const FString EntryDirectory = PDBCachePath / CleanedLabelName;
 	const FString EntryTimeStampFilename = EntryDirectory / PDBTimeStampFile;
 
-	// Verify there is a entry timestamp file.
+	// Verify there is an entry timestamp file.
 	const bool bSaved = FFileHelper::SaveStringToFile( EntryTimeStampFilename, *EntryTimeStampFilename );
 	UE_CLOG( !bSaved, LogCrashDebugHelper, Fatal, TEXT( "Couldn't save the timestamp file to %s" ), *EntryTimeStampFilename );
 	const FDateTime LastAccessTime = IFileManager::Get().GetTimeStamp( *EntryTimeStampFilename );
 
-#if	DO_LOCAL_TESTING
-	const FString RootDir = FString( TEXT( "U:/P4EPIC" ) ) / DepotName.Replace( TEXT( "//depot/" ), TEXT( "" ) );
-#else
-	const FString RootDir = FPaths::RootDir();
-#endif // DO_LOCAL_TESTING
+	const FString LocalDepotDir = DepotRoot / DepotName.Replace( P4_DEPOT_PREFIX, TEXT( "" ) );
 
-	UE_LOG( LogCrashDebugHelper, Warning, TEXT( "PDB Cache entry %s is being copied from %s, it will take some time" ), *CleanedLabelName, *OriginalLabelName );
-	for( const auto& Filename : SyncedFiles )
+	UE_LOG( LogCrashDebugHelper, Warning, TEXT( "PDB Cache entry '%s' is being copied from '%s', it will take some time" ), *CleanedLabelName, *OriginalLabelName );
+	for( const auto& Filename : FilesToBeCached )
 	{
-		const FString SourceDirectoryWithSearch = Filename.Replace( *DepotName, *RootDir );
+		const FString SourceDirectoryWithSearch = Filename.Replace( *DepotName, *LocalDepotDir );
 
 		TArray<FString> MatchedFiles;
 		IFileManager::Get().FindFiles( MatchedFiles, *SourceDirectoryWithSearch, true, false );
@@ -949,7 +1029,7 @@ FPDBCacheEntryRef FPDBCache::CreateAndAddPDBCacheEntry( const FString& OriginalL
 		for( const auto& MatchedFilename : MatchedFiles )
 		{
 			const FString SrcFilename = FPaths::GetPath( SourceDirectoryWithSearch ) / MatchedFilename;
-			const FString DestFilename = EntryDirectory / SrcFilename.Replace( *RootDir, TEXT("") );
+			const FString DestFilename = EntryDirectory / SrcFilename.Replace( *LocalDepotDir, TEXT("") );
 			IFileManager::Get().Copy( *DestFilename, *SrcFilename );
 		}
 	}
@@ -979,12 +1059,55 @@ FPDBCacheEntryRef FPDBCache::CreateAndAddPDBCacheEntry( const FString& OriginalL
 	return NewCacheEntry;
 }
 
-FPDBCacheEntryRef FPDBCache::ReadPDBCacheEntry( const FString& InLabel )
+FPDBCacheEntryRef FPDBCache::CreateAndAddPDBCacheEntryMixed( const FString& ProductVersion, const TMap<FString, FString>& FilesToBeCached )
 {
-	const FString EntryDirectory = PDBCachePath / InLabel;
+	// Enable MDD to parse all minidumps regardless the branch, to fix the issue with the missing executables on the P4 due to the build system changes.
+	const FString EntryDirectory = PDBCachePath / ProductVersion;
 	const FString EntryTimeStampFilename = EntryDirectory / PDBTimeStampFile;
+ 
+	// Verify there is an entry timestamp file.
+	const bool bSaved = FFileHelper::SaveStringToFile( EntryTimeStampFilename, *EntryTimeStampFilename );
+	UE_CLOG( !bSaved, LogCrashDebugHelper, Fatal, TEXT( "Couldn't save the timestamp file to %s" ), *EntryTimeStampFilename );
+	const FDateTime LastAccessTime = IFileManager::Get().GetTimeStamp( *EntryTimeStampFilename );
+ 
+	UE_LOG( LogCrashDebugHelper, Warning, TEXT( "PDB Cache entry '%s' is being created from %i files, it will take some time" ), *ProductVersion, FilesToBeCached.Num() );
+	for( const auto& It : FilesToBeCached )
+	{
+		const FString& SrcFilename = It.Key;
+		const FString DestFilename = EntryDirectory / It.Value;
+		IFileManager::Get().Copy( *DestFilename, *SrcFilename );
+	}
 
-	// Verify there is a entry timestamp file.
+	TArray<FString> CachedFiles;
+	IFileManager::Get().FindFilesRecursive( CachedFiles, *EntryDirectory, TEXT( "*.*" ), true, false );
+ 
+	// Calculate the size of this PDB Cache entry.
+	int64 TotalSize = 0;
+	for( const auto& Filename : CachedFiles )
+	{
+		const int64 FileSize = IFileManager::Get().FileSize( *Filename );
+		TotalSize += FileSize;
+	}
+ 
+	// Round-up the size.
+	const int32 SizeGB = (int32)FMath::DivideAndRoundUp( TotalSize, (int64)NUM_BYTES_PER_GB );
+ 
+	FPDBCacheEntryRef NewCacheEntry = MakeShareable( new FPDBCacheEntry( ProductVersion, SizeGB ) );
+	NewCacheEntry->LastAccessTime = LastAccessTime;
+	NewCacheEntry->Files = CachedFiles;
+ 
+	PDBCacheEntries.Add( ProductVersion, NewCacheEntry );
+	SortPDBCache();
+	
+	return NewCacheEntry;
+}
+
+FPDBCacheEntryRef FPDBCache::ReadPDBCacheEntry( const FString& Directory )
+{
+	const FString EntryDirectory = PDBCachePath / Directory;
+	const FString EntryTimeStampFilename = EntryDirectory / PDBTimeStampFile;
+
+	// Verify there is an entry timestamp file.
 	const FDateTime LastAccessTime = IFileManager::Get().GetTimeStamp( *EntryTimeStampFilename );
 
 	TArray<FString> PDBFiles;
@@ -1001,37 +1124,43 @@ FPDBCacheEntryRef FPDBCache::ReadPDBCacheEntry( const FString& InLabel )
 	// Round-up the size.
 	const int32 SizeGB = (int32)FMath::DivideAndRoundUp( TotalSize, (int64)NUM_BYTES_PER_GB );
 
-	FPDBCacheEntryRef NewEntry = MakeShareable( new FPDBCacheEntry( InLabel, SizeGB ) );
+	FPDBCacheEntryRef NewEntry = MakeShareable( new FPDBCacheEntry( Directory, SizeGB ) );
 	NewEntry->LastAccessTime = LastAccessTime;
 	NewEntry->Files = PDBFiles;
 	return NewEntry;
 }
 
-void FPDBCache::TouchPDBCacheEntry( const FString& InLabel )
+void FPDBCache::TouchPDBCacheEntry( const FString& Directory )
 {
-	const FString CleanedLabelName = CleanLabelName( InLabel );
-	const FString EntryDirectory = PDBCachePath / CleanedLabelName;
+	const FString EntryDirectory = PDBCachePath / Directory;
 	const FString EntryTimeStampFilename = EntryDirectory / PDBTimeStampFile;
 
-	FPDBCacheEntryRef Entry = PDBCacheEntries.FindChecked( CleanedLabelName );
+	FPDBCacheEntryRef Entry = PDBCacheEntries.FindChecked( Directory );
 	Entry->SetLastAccessTimeToNow();
 
 	const bool bResult = IFileManager::Get().SetTimeStamp( *EntryTimeStampFilename, Entry->LastAccessTime );
 	SortPDBCache();
 }
 
-void FPDBCache::RemovePDBCacheEntry( const FString& InLabel )
+void FPDBCache::RemovePDBCacheEntry( const FString& Directory )
 {
 	const double StartTime = FPlatformTime::Seconds();
 
-	const FString EntryDirectory = PDBCachePath / InLabel;
+	const FString EntryDirectory = PDBCachePath / Directory;
 
-	FPDBCacheEntryRef& Entry = PDBCacheEntries.FindChecked( InLabel );
+	FPDBCacheEntryRef& Entry = PDBCacheEntries.FindChecked( Directory );
 	IFileManager::Get().DeleteDirectory( *EntryDirectory, true, true );
-	PDBCacheEntries.Remove( InLabel );
+	PDBCacheEntries.Remove( Directory );
 
 	const double TotalTime = FPlatformTime::Seconds() - StartTime;
-	UE_LOG( LogCrashDebugHelper, Warning, TEXT( "PDB Cache entry %s removed in %.2f ms, restored %i GBs" ), *InLabel, TotalTime*1000.0f, Entry->SizeGB );
+	UE_LOG( LogCrashDebugHelper, Warning, TEXT( "PDB Cache entry %s removed in %.2f ms, restored %i GBs" ), *Directory, TotalTime*1000.0f, Entry->SizeGB );
+}
+
+FPDBCacheEntryRef FPDBCache::FindAndTouchPDBCacheEntry( const FString& PathOrLabel )
+{
+	FPDBCacheEntryRef CacheEntry = PDBCacheEntries.FindChecked( EscapePath( PathOrLabel ) );
+	TouchPDBCacheEntry( CacheEntry->Directory );
+	return CacheEntry;
 }
 
 
diff --git a/Engine/Source/Developer/CrashDebugHelper/Private/Windows/MinidumpParser.cpp b/Engine/Source/Developer/CrashDebugHelper/Private/Windows/MinidumpParser.cpp
--- a/Engine/Source/Developer/CrashDebugHelper/Private/Windows/MinidumpParser.cpp
+++ b/Engine/Source/Developer/CrashDebugHelper/Private/Windows/MinidumpParser.cpp
@@ -1,8 +1,7 @@
 // Copyright 1998-2014 Epic Games, Inc. All Rights Reserved.
 
 #include "CrashDebugHelperPrivatePCH.h"
-#include "Database.h"
-
+#include "EngineVersion.h"
 #include "Windows/WindowsPlatformStackWalkExt.h"
 #include "ISourceControlModule.h"
 
@@ -20,54 +19,51 @@ bool FCrashDebugHelperWindows::CreateMinidumpDiagnosticReport( const FString& In
 		InitSourceControl(false);
 	}
 
-	FWindowsPlatformStackWalkExt::InitStackWalking();
+	FWindowsPlatformStackWalkExt WindowsStackWalkExt( CrashInfo );
+
+	const bool bReady = WindowsStackWalkExt.InitStackWalking();
 
 	bool bAtLeastOneFunctionNameFoundInCallstack = false;
-	if( FWindowsPlatformStackWalkExt::OpenDumpFile( &CrashInfo, InCrashDumpFilename ) )
+	if( bReady && WindowsStackWalkExt.OpenDumpFile( InCrashDumpFilename ) )
 	{
-		// Get a list of modules, and the version of the executable
-		FWindowsPlatformStackWalkExt::GetModuleList( &CrashInfo );
-
-		// If a BuiltFromCL was specified, use that instead of the one parsed from the executable
-		if (BuiltFromCL > 0)
+		if( CrashInfo.ChangelistBuiltFrom != FCrashInfo::INVALID_CHANGELIST )
 		{
-			CrashInfo.ChangelistBuiltFrom = BuiltFromCL;
-		}
+			// Get the build version from the crash info.
+			FCrashModuleInfo ExeFileVersion;
+			WindowsStackWalkExt.GetExeFileVersionAndModuleList( ExeFileVersion );
+
+			// @see AutomationTool.CommandUtils.EscapePath 
+			const FString CleanedDepotName = CrashInfo.DepotName.Replace( TEXT( ":" ), TEXT( "" ) ).Replace( TEXT( "/" ), TEXT( "+" ) ).Replace( TEXT( "\\" ), TEXT( "+" ) ).Replace( TEXT( " " ), TEXT( "+" ) );
+			const FEngineVersion EngineVersion( ExeFileVersion.Major, ExeFileVersion.Minor, ExeFileVersion.Patch, CrashInfo.ChangelistBuiltFrom, CleanedDepotName );
+			CrashInfo.ProductVersion = EngineVersion.ToString();
 
-		if( CrashInfo.ChangelistBuiltFrom != -1 )
-		{
 			// CrashInfo now contains a changelist to lookup a label for
 			if( bSyncSymbols )
 			{
-				CrashInfo.LabelName = RetrieveBuildLabel( -1, CrashInfo.ChangelistBuiltFrom );
-				if( CrashInfo.LabelName.Len() > 0 )
-				{
-					// Sync all the modules, and their associated pdbs
-					CrashInfo.Log( FString::Printf( TEXT( " ... found label '%s'; syncing all modules to this label" ), *CrashInfo.LabelName ) );
-					SyncModules( &CrashInfo );
-				}
+				RetrieveBuildLabelAndNetworkPath( CrashInfo.ChangelistBuiltFrom );
+				SyncModules();
 			}
 
 			// Initialise the symbol options
-			FWindowsPlatformStackWalkExt::InitSymbols( &CrashInfo );
+			WindowsStackWalkExt.InitSymbols();
 
 			// Set the symbol path based on the loaded modules
-			FWindowsPlatformStackWalkExt::SetSymbolPathsFromModules( &CrashInfo );
+			WindowsStackWalkExt.SetSymbolPathsFromModules();
 
 			// Get all the info we should ever need about the modules
-			FWindowsPlatformStackWalkExt::GetModuleInfoDetailed( &CrashInfo );
+			WindowsStackWalkExt.GetModuleInfoDetailed();
 
 			// Get info about the system that created the minidump
-			FWindowsPlatformStackWalkExt::GetSystemInfo( &CrashInfo );
+			WindowsStackWalkExt.GetSystemInfo();
 
 			// Get all the thread info
-			FWindowsPlatformStackWalkExt::GetThreadInfo( &CrashInfo );
+			WindowsStackWalkExt.GetThreadInfo();
 
 			// Get exception info
-			FWindowsPlatformStackWalkExt::GetExceptionInfo( &CrashInfo );
+			WindowsStackWalkExt.GetExceptionInfo();
 
 			// Get the callstacks for each thread
-			bAtLeastOneFunctionNameFoundInCallstack = FWindowsPlatformStackWalkExt::GetCallstacks(&CrashInfo);
+			bAtLeastOneFunctionNameFoundInCallstack = WindowsStackWalkExt.GetCallstacks();
 
 			// Sync the source file where the crash occurred
 			if( CrashInfo.SourceFile.Len() > 0 )
@@ -75,20 +71,20 @@ bool FCrashDebugHelperWindows::CreateMinidumpDiagnosticReport( const FString& In
 				if( bSyncSymbols )
 				{
 					CrashInfo.Log( FString::Printf( TEXT( " ... using label '%s' to sync crash source file" ), *CrashInfo.LabelName ) );
-					SyncSourceFile( &CrashInfo );
+					SyncSourceFile();
 				}
 
 				// Try to annotate the file if requested
 				bool bAnnotationSuccessful = false;
 				if( bAnnotate )
 				{
-					bAnnotationSuccessful = AddAnnotatedSourceToReport( &CrashInfo );
+					bAnnotationSuccessful = AddAnnotatedSourceToReport();
 				}
 				
 				// If annotation is not requested, or failed, add the standard source context
 				if( !bAnnotationSuccessful )
 				{
-					AddSourceToReport( &CrashInfo );
+					AddSourceToReport();
 				}
 			}
 		}
@@ -102,8 +98,6 @@ bool FCrashDebugHelperWindows::CreateMinidumpDiagnosticReport( const FString& In
 		CrashInfo.Log( FString::Printf( TEXT( "CreateMinidumpDiagnosticReport: Failed to open crash dump file: %s" ), *InCrashDumpFilename ) );
 	}
 
-	FWindowsPlatformStackWalkExt::ShutdownStackWalking();
-
 	if( bUseSCC )
 	{
 		ShutdownSourceControl();
diff --git a/Engine/Source/Developer/CrashDebugHelper/Private/Windows/WindowsPlatformStackWalkExt.cpp b/Engine/Source/Developer/CrashDebugHelper/Private/Windows/WindowsPlatformStackWalkExt.cpp
--- a/Engine/Source/Developer/CrashDebugHelper/Private/Windows/WindowsPlatformStackWalkExt.cpp
+++ b/Engine/Source/Developer/CrashDebugHelper/Private/Windows/WindowsPlatformStackWalkExt.cpp
@@ -12,13 +12,21 @@
 
 #pragma comment( lib, "dbgeng.lib" )
 
-static FCrashInfo* StaticCrashInfo = NULL;
-
 static IDebugClient5* Client = NULL;
 static IDebugControl4* Control = NULL;
 static IDebugSymbols3* Symbol = NULL;
 static IDebugAdvanced3* Advanced = NULL;
 
+FWindowsPlatformStackWalkExt::FWindowsPlatformStackWalkExt( FCrashInfo& InCrashInfo ) 
+	: CrashInfo( InCrashInfo )
+{}
+
+FWindowsPlatformStackWalkExt::~FWindowsPlatformStackWalkExt()
+{
+	ShutdownStackWalking();
+}
+
+
 /**
  * Initialise the COM interface to grab stacks
  */
@@ -55,7 +63,7 @@ void FWindowsPlatformStackWalkExt::ShutdownStackWalking()
 /**
  * Set the options we want for symbol lookup
  */
-void FWindowsPlatformStackWalkExt::InitSymbols( FCrashInfo* CrashInfo )
+void FWindowsPlatformStackWalkExt::InitSymbols()
 {
 	ULONG SymOpts = 0;
 
@@ -111,19 +119,14 @@ FString FWindowsPlatformStackWalkExt::ExtractRelativePath( const TCHAR* BaseName
 	return FullPath;
 }
 
-/** 
- * Get a list of module names and versions
- */
-void FWindowsPlatformStackWalkExt::GetModuleList( FCrashInfo* CrashInfo )
+void FWindowsPlatformStackWalkExt::GetExeFileVersionAndModuleList( FCrashModuleInfo& out_ExeFileVersion )
 {
-	CrashInfo->Log( FString::Printf( TEXT( "Getting list of modules." ) ) );
-
 	// The the number of loaded modules
 	ULONG LoadedModuleCount = 0;
 	ULONG UnloadedModuleCount = 0;
 	Symbol->GetNumberModules( &LoadedModuleCount, &UnloadedModuleCount );
 
-	CrashInfo->Log( FString::Printf( TEXT( " ... found %d loaded, and %d unloaded modules." ), LoadedModuleCount, UnloadedModuleCount ) );
+	CrashInfo.Log( FString::Printf( TEXT( " ... found %d loaded, and %d unloaded modules." ), LoadedModuleCount, UnloadedModuleCount ) );
 
 	// Find the relative names of all the modules so we know which files to sync
 	int ExecutableIndex = -1;
@@ -133,13 +136,13 @@ void FWindowsPlatformStackWalkExt::GetModuleList( FCrashInfo* CrashInfo )
 		Symbol->GetModuleByIndex( ModuleIndex, &ModuleBase );
 
 		// Get the full path of the module name
-		TCHAR ModuleName[MAX_PATH] = { 0 };
+		TCHAR ModuleName[MAX_PATH] = {0};
 		Symbol->GetModuleNameStringWide( DEBUG_MODNAME_IMAGE, ModuleIndex, ModuleBase, ModuleName, MAX_PATH, NULL );
 
 		FString RelativeModuleName = ExtractRelativePath( TEXT( "binaries" ), ModuleName );
 		if( RelativeModuleName.Len() > 0 )
 		{
-			CrashInfo->ModuleNames.Add( RelativeModuleName );
+			CrashInfo.ModuleNames.Add( RelativeModuleName );
 		}
 
 		// Get the exe, which we extract the version number, so we know what label to sync to
@@ -149,45 +152,35 @@ void FWindowsPlatformStackWalkExt::GetModuleList( FCrashInfo* CrashInfo )
 		}
 	}
 
-	// Group all the folders together
-	CrashInfo->ModuleNames.Sort();
-
-	// Get the executable version info
+	// Get the executable version info.
 	if( ExecutableIndex > -1 )
 	{
-		VS_FIXEDFILEINFO VersionInfo = { 0 };
+		VS_FIXEDFILEINFO VersionInfo = {0};
 		Symbol->GetModuleVersionInformationWide( ExecutableIndex, 0, TEXT( "\\" ), &VersionInfo, sizeof( VS_FIXEDFILEINFO ), NULL );
 
-		uint32 Build = HIWORD( VersionInfo.dwProductVersionLS );
-		uint32 Revision = LOWORD( VersionInfo.dwProductVersionLS );
-
-		if( Revision == 0 )
-		{
-			CrashInfo->ChangelistBuiltFrom = Build;
-		}
-		else
-		{
-			CrashInfo->ChangelistBuiltFrom = ( Build << 16 ) | Revision;
-		}
+		out_ExeFileVersion.Major = HIWORD( VersionInfo.dwProductVersionMS );
+		out_ExeFileVersion.Minor = LOWORD( VersionInfo.dwProductVersionMS );
+		out_ExeFileVersion.Patch = HIWORD( VersionInfo.dwProductVersionLS );
 	}
 	else
 	{
-		CrashInfo->Log( FString::Printf( TEXT( " ... unable to locate executable." ), LoadedModuleCount, UnloadedModuleCount ) );
+		CrashInfo.Log( FString::Printf( TEXT( " ... unable to locate executable." ) ) );
 	}
 }
 
 /** 
  * Set the symbol paths based on the module paths
  */
-void FWindowsPlatformStackWalkExt::SetSymbolPathsFromModules( FCrashInfo* CrashInfo )
+void FWindowsPlatformStackWalkExt::SetSymbolPathsFromModules()
 {
-	const bool bUseCachedData = CrashInfo->PDBCacheEntry.IsValid();
+	const bool bUseCachedData = CrashInfo.PDBCacheEntry.IsValid();
 	FString CombinedPath = TEXT( "" );
 
+	// For externally launched minidump diagnostics.
 	if( bUseCachedData )
 	{
 		TSet<FString> SymbolPaths;
-		for( const auto& Filename : CrashInfo->PDBCacheEntry->Files )
+		for( const auto& Filename : CrashInfo.PDBCacheEntry->Files )
 		{
 			const FString SymbolPath = FPaths::GetPath( Filename );
 			if( SymbolPath.Len() > 0 )
@@ -206,72 +199,34 @@ void FWindowsPlatformStackWalkExt::SetSymbolPathsFromModules( FCrashInfo* CrashI
 		Symbol->SetImagePathWide( *CombinedPath );
 		Symbol->SetSymbolPathWide( *CombinedPath );
 	}
+	// For locally launched minidump diagnostics.
 	else
 	{
-		const bool bUseNewSetPath = true;
-		if( bUseNewSetPath )
-		{
-#if	DO_LOCAL_TESTING
-			FString DepotName;
-			FParse::Value( FCommandLine::Get(), TEXT( "BranchName=" ), DepotName );
-			const FString RootDir = FString( TEXT( "U:/P4EPIC" ) ) / DepotName;
-#else
-			const FString RootDir = FPaths::RootDir();
-#endif // DO_LOCAL_TESTING
-
-			TArray<FString> ModulesAndSymbols;
-			IFileManager::Get().FindFilesRecursive( ModulesAndSymbols, *RootDir, TEXT( "*.pdb" ), true, false, false );
-			IFileManager::Get().FindFilesRecursive( ModulesAndSymbols, *RootDir, TEXT( "*.dll" ), true, false, false );
-
-			TSet<FString> SymbolPaths;
-			for( const auto& Filename : ModulesAndSymbols )
-			{
-				const FString Path = FPaths::GetPath( Filename );
-				if( Path.Len() > 0 )
-				{
-					SymbolPaths.Add( Path );
-				}
-			}
-
-			for( const auto& SymbolPath : SymbolPaths )
-			{
-				CombinedPath += SymbolPath;
-				CombinedPath += TEXT( ";" );
-			}
+		// FindFilesRecursive for root is quite slow.
+		TArray<FString> ModulesAndSymbols;
+		IFileManager::Get().FindFilesRecursive( ModulesAndSymbols, *FPaths::RootDir(), TEXT( "*.pdb" ), true, false, false );
+		IFileManager::Get().FindFilesRecursive( ModulesAndSymbols, *FPaths::RootDir(), TEXT( "*.dll" ), true, false, false );
+	
 
-			// Set the symbol path
-			Symbol->SetImagePathWide( *CombinedPath );
-			Symbol->SetSymbolPathWide( *CombinedPath );
-		}
-		else
+		TSet<FString> SymbolPaths;
+		for( const auto& Filename : ModulesAndSymbols )
 		{
-			TArray<FString> SymbolPaths;
-			for( int32 ModuleNameIndex = 0; ModuleNameIndex < CrashInfo->ModuleNames.Num(); ModuleNameIndex++ )
+			const FString Path = FPaths::GetPath( Filename );
+			if( Path.Len() > 0 )
 			{
-				FString ModulePath = FPaths::GetPath( CrashInfo->ModuleNames[ModuleNameIndex] );
-				if( ModulePath.Len() > 0 )
-				{
-					SymbolPaths.AddUnique( ModulePath );
-				}
+				SymbolPaths.Add( Path );
 			}
+		}
 
-			for( int32 PathIndex = 0; PathIndex < SymbolPaths.Num(); PathIndex++ )
-			{
-				CombinedPath += FString( TEXT( "../../../" ) ) + SymbolPaths[PathIndex] + TEXT( ";" );
-			}
+		for( const auto& SymbolPath : SymbolPaths )
+		{
+			CombinedPath += SymbolPath;
+			CombinedPath += TEXT( ";" );
+		}
 
-			// Set the symbol path
-			Symbol->SetImagePathWide( *CombinedPath );
-			Symbol->SetSymbolPathWide( *CombinedPath );
-
-			//@TODO: ROCKETHACK: This is for Rocket only.
-			Symbol->AppendImagePathWide( TEXT( "..\\..\\..\\Rocket\\Installed\\Windows\\Engine\\Binaries\\Win64" ) );
-			Symbol->AppendImagePathWide( TEXT( "..\\..\\..\\Rocket\\Installed\\Windows\\Engine\\Binaries\\Win32" ) );
-			Symbol->AppendSymbolPathWide( TEXT( "..\\..\\..\\Rocket\\Symbols\\Engine\\Binaries\\Win64" ) );
-			Symbol->AppendSymbolPathWide( TEXT( "..\\..\\..\\Rocket\\Symbols\\Engine\\Binaries\\Win32" ) );
-			Symbol->AppendImagePathWide( TEXT( "..\\..\\..\\Rocket\\LauncherInstalled\\Windows\\Launcher\\Engine\\Binaries\\Win64" ) );
-			Symbol->AppendSymbolPathWide( TEXT( "..\\..\\..\\Rocket\\LauncherSymbols\\Windows\\Launcher\\Engine\\Binaries\\Win64" ) );
-		}	
+		// Set the symbol path
+		Symbol->SetImagePathWide( *CombinedPath );
+		Symbol->SetSymbolPathWide( *CombinedPath );
 	}
 
 	// Add in syncing of the Microsoft symbol servers if requested
@@ -280,15 +235,15 @@ void FWindowsPlatformStackWalkExt::SetSymbolPathsFromModules( FCrashInfo* CrashI
 		FString BinariesDir = FString(FPlatformProcess::BaseDir());
 		if ( !FPaths::FileExists( FPaths::Combine(*BinariesDir, TEXT("symsrv.dll")) ) )
 		{
-			CrashInfo->Log( FString::Printf(TEXT("Error: symsrv.dll was not detected in '%s'. Microsoft symbols will not be downloaded!"), *BinariesDir) );
+			CrashInfo.Log( FString::Printf(TEXT("Error: symsrv.dll was not detected in '%s'. Microsoft symbols will not be downloaded!"), *BinariesDir) );
 		}
 		if ( !FPaths::FileExists( FPaths::Combine(*BinariesDir, TEXT("symsrv.yes")) ) )
 		{
-			CrashInfo->Log( FString::Printf(TEXT("symsrv.yes was not detected in '%s'. This will cause a popup to confirm the license."), *BinariesDir) );
+			CrashInfo.Log( FString::Printf(TEXT("symsrv.yes was not detected in '%s'. This will cause a popup to confirm the license."), *BinariesDir) );
 		}
 		if ( !FPaths::FileExists( FPaths::Combine(*BinariesDir, TEXT("dbghelp.dll")) ) )
 		{
-			CrashInfo->Log( FString::Printf(TEXT("Error: dbghelp.dll was not detected in '%s'. Microsoft symbols will not be downloaded!"), *BinariesDir) );
+			CrashInfo.Log( FString::Printf(TEXT("Error: dbghelp.dll was not detected in '%s'. Microsoft symbols will not be downloaded!"), *BinariesDir) );
 		}
 
 		Symbol->AppendImagePathWide( TEXT( "SRV*..\\..\\Intermediate\\SymbolCache*http://msdl.microsoft.com/download/symbols" ) );
@@ -297,10 +252,10 @@ void FWindowsPlatformStackWalkExt::SetSymbolPathsFromModules( FCrashInfo* CrashI
 
 	TCHAR SymbolPath[1024] = { 0 };
 	Symbol->GetSymbolPathWide( SymbolPath, 1024, NULL );
-	CrashInfo->Log( FString::Printf( TEXT( " ... setting symbol path to: '%s'" ), SymbolPath ) );
+	CrashInfo.Log( FString::Printf( TEXT( " ... setting symbol path to: '%s'" ), SymbolPath ) );
 
 	Symbol->GetImagePathWide( SymbolPath, 1024, NULL );
-	CrashInfo->Log( FString::Printf( TEXT( " ... setting image path to: '%s'" ), SymbolPath ) );
+	CrashInfo.Log( FString::Printf( TEXT( " ... setting image path to: '%s'" ), SymbolPath ) );
 }
 
 class FSortModulesByName
@@ -327,17 +282,17 @@ public:
 /** 
  * Get detailed info about each module
  */
-void FWindowsPlatformStackWalkExt::GetModuleInfoDetailed( FCrashInfo* CrashInfo )
+void FWindowsPlatformStackWalkExt::GetModuleInfoDetailed()
 {
 	// The the number of loaded modules
 	ULONG LoadedModuleCount = 0;
 	ULONG UnloadedModuleCount = 0;
 	Symbol->GetNumberModules( &LoadedModuleCount, &UnloadedModuleCount );
 
-	CrashInfo->Modules.Empty( LoadedModuleCount );
+	CrashInfo.Modules.Empty( LoadedModuleCount );
 	for( uint32 ModuleIndex = 0; ModuleIndex < LoadedModuleCount; ModuleIndex++ )
 	{
-		FCrashModuleInfo* CrashModule = new ( CrashInfo->Modules ) FCrashModuleInfo();
+		FCrashModuleInfo* CrashModule = new ( CrashInfo.Modules ) FCrashModuleInfo();
 
 		ULONG64 ModuleBase = 0;
 		Symbol->GetModuleByIndex( ModuleIndex, &ModuleBase );
@@ -367,24 +322,24 @@ void FWindowsPlatformStackWalkExt::GetModuleInfoDetailed( FCrashInfo* CrashInfo
 
 		CrashModule->Major = HIWORD( VersionInfo.dwProductVersionMS );
 		CrashModule->Minor = LOWORD( VersionInfo.dwProductVersionMS );
-		CrashModule->Build = HIWORD( VersionInfo.dwProductVersionLS );
+		CrashModule->Patch = HIWORD( VersionInfo.dwProductVersionLS );
 		CrashModule->Revision = LOWORD( VersionInfo.dwProductVersionLS );
 
 		// Ensure all the images are synced - need the full path here
 		Symbol->ReloadWide( *CrashModule->Name );
 	}
 
-	CrashInfo->Modules.Sort( FSortModulesByName() );
+	CrashInfo.Modules.Sort( FSortModulesByName() );
 }
 
 /**
  * Check to see if the stack address resides within one of the loaded modules i.e. is it valid?
  */
-bool FWindowsPlatformStackWalkExt::IsOffsetWithinModules( FCrashInfo* CrashInfo, uint64 Offset )
+bool FWindowsPlatformStackWalkExt::IsOffsetWithinModules( uint64 Offset )
 {
-	for( int ModuleIndex = 0; ModuleIndex < CrashInfo->Modules.Num(); ModuleIndex++ )
+	for( int ModuleIndex = 0; ModuleIndex < CrashInfo.Modules.Num(); ModuleIndex++ )
 	{
-		FCrashModuleInfo& CrashModule = CrashInfo->Modules[ModuleIndex];
+		FCrashModuleInfo& CrashModule = CrashInfo.Modules[ModuleIndex];
 		if( Offset >= CrashModule.BaseOfImage && Offset < CrashModule.BaseOfImage + CrashModule.SizeOfImage )
 		{
 			return true;
@@ -397,9 +352,9 @@ bool FWindowsPlatformStackWalkExt::IsOffsetWithinModules( FCrashInfo* CrashInfo,
 /** 
  * Extract the system info of the crash from the minidump
  */
-void FWindowsPlatformStackWalkExt::GetSystemInfo( FCrashInfo* CrashInfo )
+void FWindowsPlatformStackWalkExt::GetSystemInfo()
 {
-	FCrashSystemInfo& SystemInfo = CrashInfo->SystemInfo;
+	FCrashSystemInfo& SystemInfo = CrashInfo.SystemInfo;
 
 	ULONG PlatformId = 0;
 	ULONG Major = 0;
@@ -420,17 +375,17 @@ void FWindowsPlatformStackWalkExt::GetSystemInfo( FCrashInfo* CrashInfo )
 	{
 	case IMAGE_FILE_MACHINE_I386:
 		// x86
-		CrashInfo->SystemInfo.ProcessorArchitecture = PA_X86;
+		CrashInfo.SystemInfo.ProcessorArchitecture = PA_X86;
 		break;
 
 	case IMAGE_FILE_MACHINE_ARM:
 		// ARM
-		CrashInfo->SystemInfo.ProcessorArchitecture = PA_ARM;
+		CrashInfo.SystemInfo.ProcessorArchitecture = PA_ARM;
 		break;
 
 	case IMAGE_FILE_MACHINE_AMD64:
 		// x64
-		CrashInfo->SystemInfo.ProcessorArchitecture = PA_X64;
+		CrashInfo.SystemInfo.ProcessorArchitecture = PA_X64;
 		break;
 
 	default: 
@@ -445,16 +400,16 @@ void FWindowsPlatformStackWalkExt::GetSystemInfo( FCrashInfo* CrashInfo )
 /** 
  * Extract the thread info from the minidump
  */
-void FWindowsPlatformStackWalkExt::GetThreadInfo( FCrashInfo* CrashInfo )
+void FWindowsPlatformStackWalkExt::GetThreadInfo()
 {
 }
 
 /** 
  * Extract info about the exception that caused the crash
  */
-void FWindowsPlatformStackWalkExt::GetExceptionInfo( FCrashInfo* CrashInfo )
+void FWindowsPlatformStackWalkExt::GetExceptionInfo()
 {
-	FCrashExceptionInfo& Exception = CrashInfo->Exception;
+	FCrashExceptionInfo& Exception = CrashInfo.Exception;
 
 	ULONG ExceptionType = 0;
 	ULONG ProcessID = 0;
@@ -471,11 +426,11 @@ void FWindowsPlatformStackWalkExt::GetExceptionInfo( FCrashInfo* CrashInfo )
 /** 
  * Get the callstack of the crash. Returns true if at least one function name was found in the debugging symbols
  */
-bool FWindowsPlatformStackWalkExt::GetCallstacks( FCrashInfo* CrashInfo )
+bool FWindowsPlatformStackWalkExt::GetCallstacks()
 {
 	bool bAtLeastOneFunctionNameFound = false;
 
-	FCrashExceptionInfo& Exception = CrashInfo->Exception;
+	FCrashExceptionInfo& Exception = CrashInfo.Exception;
 
 	byte Context[4096] = { 0 };
 	ULONG DebugEvent = 0;
@@ -492,11 +447,11 @@ bool FWindowsPlatformStackWalkExt::GetCallstacks( FCrashInfo* CrashInfo )
 	// Some magic number checks
 	if( ContextSize == 716 )
 	{
-		CrashInfo->Log( FString::Printf( TEXT( " ... context size matches x86 sizeof( CONTEXT )" ) ) );
+		CrashInfo.Log( FString::Printf( TEXT( " ... context size matches x86 sizeof( CONTEXT )" ) ) );
 	}
 	else if( ContextSize == 1232 )
 	{
-		CrashInfo->Log( FString::Printf( TEXT( " ... context size matches x64 sizeof( CONTEXT )" ) ) );
+		CrashInfo.Log( FString::Printf( TEXT( " ... context size matches x64 sizeof( CONTEXT )" ) ) );
 	}
 
 	// Get the entire stack trace
@@ -515,7 +470,7 @@ bool FWindowsPlatformStackWalkExt::GetCallstacks( FCrashInfo* CrashInfo )
 		FString Line;
 		uint64 Offset = StackFrames[StackIndex].InstructionOffset;
 
-		if( FWindowsPlatformStackWalkExt::IsOffsetWithinModules( CrashInfo, Offset ) )
+		if( IsOffsetWithinModules( Offset ) )
 		{
 			// Get the module, function, and offset
 			uint64 Displacement = 0;
@@ -534,8 +489,8 @@ bool FWindowsPlatformStackWalkExt::GetCallstacks( FCrashInfo* CrashInfo )
 			// Remember the top of the stack to locate in the source file
 			if( !bFoundSourceFile && FString( SourceName ).Len() > 0 && LineNumber > 0 )
 			{
-				CrashInfo->SourceFile = ExtractRelativePath( TEXT( "source" ), SourceName );
-				CrashInfo->SourceLineNumber = LineNumber;
+				CrashInfo.SourceFile = ExtractRelativePath( TEXT( "source" ), SourceName );
+				CrashInfo.SourceLineNumber = LineNumber;
 				bFoundSourceFile = true;
 			}
 
@@ -577,7 +532,7 @@ bool FWindowsPlatformStackWalkExt::GetCallstacks( FCrashInfo* CrashInfo )
 
 	FMemory::Free( ContextData );
 
-	CrashInfo->Log( FString::Printf( TEXT( " ... callstack generated!" ) ) );
+	CrashInfo.Log( FString::Printf( TEXT( " ... callstack generated!" ) ) );
 
 	return bAtLeastOneFunctionNameFound;
 }
@@ -585,30 +540,27 @@ bool FWindowsPlatformStackWalkExt::GetCallstacks( FCrashInfo* CrashInfo )
 /** 
  * Open a minidump as a new session
  */
-bool FWindowsPlatformStackWalkExt::OpenDumpFile( FCrashInfo* CrashInfo, const FString& InCrashDumpFilename )
+bool FWindowsPlatformStackWalkExt::OpenDumpFile( const FString& InCrashDumpFilename )
 {
-	HRESULT Result = S_OK;
-
-	HANDLE DumpHandle = CreateFileW( *InCrashDumpFilename, GENERIC_READ, FILE_SHARE_READ, NULL, OPEN_EXISTING, 0, NULL );
-	if( DumpHandle == INVALID_HANDLE_VALUE )
+	const bool bFound = IFileManager::Get().FileSize( *InCrashDumpFilename ) != INDEX_NONE;
+	if( !bFound )
 	{
-		CrashInfo->Log( FString::Printf( TEXT( " ... failed to find minidump file: %s" ), *InCrashDumpFilename ) );
+		CrashInfo.Log( FString::Printf( TEXT( " ... failed to find minidump file: %s" ), *InCrashDumpFilename ) );
 		return false;
 	}
 
-	wchar_t DumpName[MAX_PATH] = { 0 };
-	if( Client->OpenDumpFileWide( DumpName, ( ULONG64 )DumpHandle ) )
+	if( Client->OpenDumpFileWide( *InCrashDumpFilename, NULL ) )
 	{
-		CrashInfo->Log( FString::Printf( TEXT( " ... failed to open minidump: %s" ), *InCrashDumpFilename ) );
+		CrashInfo.Log( FString::Printf( TEXT( " ... failed to open minidump: %s" ), *InCrashDumpFilename ) );
 		return false;
 	}
 
 	if( Control->WaitForEvent( 0, INFINITE ) != S_OK )
 	{
-		CrashInfo->Log( FString::Printf( TEXT( " ... failed while waiting for minidump to load: %s" ), *InCrashDumpFilename ) );
+		CrashInfo.Log( FString::Printf( TEXT( " ... failed while waiting for minidump to load: %s" ), *InCrashDumpFilename ) );
 		return false;
 	}
 
-	CrashInfo->Log( FString::Printf( TEXT( "Successfully opened minidump: %s" ), *InCrashDumpFilename ) );
+	CrashInfo.Log( FString::Printf( TEXT( "Successfully opened minidump: %s" ), *InCrashDumpFilename ) );
 	return true;
 }
diff --git a/Engine/Source/Developer/CrashDebugHelper/Private/Windows/WindowsPlatformStackWalkExt.h b/Engine/Source/Developer/CrashDebugHelper/Private/Windows/WindowsPlatformStackWalkExt.h
--- a/Engine/Source/Developer/CrashDebugHelper/Private/Windows/WindowsPlatformStackWalkExt.h
+++ b/Engine/Source/Developer/CrashDebugHelper/Private/Windows/WindowsPlatformStackWalkExt.h
@@ -1,32 +1,42 @@
 // Copyright 1998-2014 Epic Games, Inc. All Rights Reserved.
 
-#ifndef __WindowsPlatformStackWalkExt_h__
-#define __WindowsPlatformStackWalkExt_h__
 #pragma once
 
 #undef PLATFORM_SUPPORTS_STACK_SYMBOLS
 #define PLATFORM_SUPPORTS_STACK_SYMBOLS 1
 
+class FCrashInfo;
+class FCrashModuleInfo;
+
 /**
- * Windows implementation of stackwalking
- **/
-struct CORE_API FWindowsPlatformStackWalkExt : public FGenericPlatformStackWalk
+ *	Windows implementation of stack walking using the COM interface IDebugClient5. 
+ */
+struct FWindowsPlatformStackWalkExt
 {
-	static bool InitStackWalking();
-	static void ShutdownStackWalking();
-
-	static void InitSymbols( class FCrashInfo* CrashInfo );
-	static bool OpenDumpFile( FCrashInfo* CrashInfo, const FString& DumpFileName );
-	static void GetModuleList( FCrashInfo* CrashInfo );
-	static void SetSymbolPathsFromModules( FCrashInfo* CrashInfo );
-	static void GetModuleInfoDetailed( FCrashInfo* CrashInfo );
-	static void GetSystemInfo( FCrashInfo* CrashInfo );
-	static void GetThreadInfo( FCrashInfo* CrashInfo );
-	static void GetExceptionInfo( FCrashInfo* CrashInfo );
-	static bool GetCallstacks( FCrashInfo* CrashInfo );
-
-	static bool IsOffsetWithinModules( FCrashInfo* CrashInfo, uint64 Offset );
+	/** Default constructor. */
+	FWindowsPlatformStackWalkExt( FCrashInfo& InCrashInfo );
+
+	/** Destructor. */
+	~FWindowsPlatformStackWalkExt();
+
+	bool InitStackWalking();
+	void ShutdownStackWalking();
+
+	void GetExeFileVersionAndModuleList( FCrashModuleInfo& out_ExeFileVersion );
+	void InitSymbols();
+	bool OpenDumpFile( const FString& DumpFileName );
+	void SetSymbolPathsFromModules();
+	void GetModuleInfoDetailed();
+	void GetSystemInfo();
+	void GetThreadInfo();
+	void GetExceptionInfo();
+	bool GetCallstacks();
+
+	bool IsOffsetWithinModules( uint64 Offset );
+
 	static FString ExtractRelativePath( const TCHAR* BaseName, TCHAR* FullName );
-};
 
-#endif
\ No newline at end of file
+protected:
+	/** Reference to the crash info. */
+	FCrashInfo& CrashInfo;
+};
diff --git a/Engine/Source/Developer/CrashDebugHelper/Public/CrashDebugHelper.h b/Engine/Source/Developer/CrashDebugHelper/Public/CrashDebugHelper.h
--- a/Engine/Source/Developer/CrashDebugHelper/Public/CrashDebugHelper.h
+++ b/Engine/Source/Developer/CrashDebugHelper/Public/CrashDebugHelper.h
@@ -2,8 +2,6 @@
 
 #pragma once
 
-//#define DO_LOCAL_TESTING 1
-
 struct FPDBCacheEntry;
 typedef TSharedRef<FPDBCacheEntry> FPDBCacheEntryRef;
 typedef TSharedPtr<FPDBCacheEntry> FPDBCacheEntryPtr;
@@ -30,7 +28,7 @@ public:
 	uint32 SizeOfImage;
 	uint16 Major;
 	uint16 Minor;
-	uint16 Build;
+	uint16 Patch;
 	uint16 Revision;
 
 	FCrashModuleInfo()
@@ -38,7 +36,7 @@ public:
 		, SizeOfImage( 0 )
 		, Major( 0 )
 		, Minor( 0 )
-		, Build( 0 )
+		, Patch( 0 )
 		, Revision( 0 )
 	{
 	}
@@ -126,16 +124,37 @@ public:
 class CRASHDEBUGHELPER_API FCrashInfo
 {
 public:
+	enum
+	{
+		/** An invalid changelist, something went wrong. */
+		INVALID_CHANGELIST = -1,
+	};
+
+	/** Report log. */
 	FString Report;
 
-	TArray<FString> ModuleNames;
+	/** The depot name, indicate where the executables and symbols are stored. */
+	FString DepotName;
+
+	/** Product version, based on FEngineVersion. */
+	FString ProductVersion;
+
+	/** CL built from. */
 	int32 ChangelistBuiltFrom;
+
+	/** The label the describes the executables and symbols. */
 	FString LabelName;
 
+	/** The network path where the executables are stored. */
+	FString NetworkPath;
+
 	FString SourceFile;
 	uint64 SourceLineNumber;
 	TArray<FString> SourceContext;
 
+	/** Only modules names, retrieved from the minidump file. */
+	TArray<FString> ModuleNames;
+
 	FCrashSystemInfo SystemInfo;
 	FCrashExceptionInfo Exception;
 	TArray<FCrashThreadInfo> Threads;
@@ -145,7 +164,7 @@ public:
 	FPDBCacheEntryPtr PDBCacheEntry;
 
 	FCrashInfo()
-		: ChangelistBuiltFrom( -1 )
+		: ChangelistBuiltFrom( INVALID_CHANGELIST )
 		, SourceLineNumber( 0 )
 	{
 	}
@@ -187,8 +206,8 @@ struct CRASHDEBUGHELPER_API FCrashDebugInfo
 struct FPDBCacheEntry
 {
 	/** Initialization constructor. */
-	FPDBCacheEntry( const FString& InLabel, const int32 InSizeGB )
-		: Label( InLabel )
+	FPDBCacheEntry( const FString& InDirectory, const int32 InSizeGB )
+		: Directory( InDirectory )
 		, SizeGB( InSizeGB )
 	{}
 
@@ -200,8 +219,8 @@ struct FPDBCacheEntry
 	/** Paths to files associated with this PDB Cache entry. */
 	TArray<FString> Files;
 
-	/** P4 Label associated with this PDB Cache entry. */
-	const FString Label;
+	/** The path associated with this PDB Cache entry. */
+	const FString Directory;
 	
 	/** Last access time, changed every time this PDB cache entry is used. */
 	FDateTime LastAccessTime;
@@ -221,7 +240,6 @@ struct FPDBCacheEntryByAccessTime
 /** Implements PDB cache. */
 struct FPDBCache
 {
-	//friend class ICrashDebugHelper;
 protected:
 
 	// Defaults for the PDB cache.
@@ -252,9 +270,6 @@ protected:
 	/** Path to the folder where the PDB cache is stored. */
 	FString PDBCachePath;
 
-	/** The builtFromCL to use instead of extracting from the minidump */
-	int32 BuiltFromCL;
-
 	/** Age of file when it should be deleted from the PDB cache. */
 	int32 DaysToDeleteUnusedFilesFromPDBCache;
 
@@ -274,8 +289,7 @@ protected:
 public:
 	/** Default constructor. */
 	FPDBCache()
-		: BuiltFromCL( 0 )
-		, DaysToDeleteUnusedFilesFromPDBCache( DAYS_TO_DELETE_UNUSED_FILES )
+		: DaysToDeleteUnusedFilesFromPDBCache( DAYS_TO_DELETE_UNUSED_FILES )
 		, PDBCacheSizeGB( PDB_CACHE_SIZE_GB )
 		, MinDiskFreeSpaceGB( MIN_FREESPACE_GB )
 		, bUsePDBCache( false )
@@ -295,28 +309,30 @@ public:
 	/**
 	 * @return true, if the PDB Cache contains the specified label.
 	 */
-	bool ContainsPDBCacheEntry( const FString& InLabel ) const
+	bool ContainsPDBCacheEntry( const FString& PathOrLabel ) const
 	{
-		return PDBCacheEntries.Contains( CleanLabelName( InLabel ) );
+		return PDBCacheEntries.Contains( EscapePath( PathOrLabel ) );
 	}
 
 	/**
 	 *	Touches a PDB Cache entry by updating the timestamp.
 	 */
-	void TouchPDBCacheEntry( const FString& InLabel );
+	void TouchPDBCacheEntry( const FString& Directory );
 
 	/**
-	 * @return a PDB Cache entry for the specified label.
+	 * @return a PDB Cache entry for the specified label and touches it at the same time
 	 */
-	FPDBCacheEntryRef FindPDBCacheEntry( const FString& InLabel )
-	{
-		return PDBCacheEntries.FindChecked( CleanLabelName( InLabel ) );
-	}
+	FPDBCacheEntryRef FindAndTouchPDBCacheEntry( const FString& PathOrLabel );
+
+	/**
+	 *	Creates a new PDB Cache entry, initializes it and adds to the database.
+	 */
+	FPDBCacheEntryRef CreateAndAddPDBCacheEntry( const FString& OriginalLabelName, const FString& DepotRoot, const FString& DepotName, const TArray<FString>& FilesToBeCached );
 
 	/**
 	 *	Creates a new PDB Cache entry, initializes it and adds to the database.
 	 */
-	FPDBCacheEntryRef CreateAndAddPDBCacheEntry( const FString& OriginalLabelName, const FString& DepotName, const TArray<FString>& SyncedFiles );
+	FPDBCacheEntryRef CreateAndAddPDBCacheEntryMixed( const FString& ProductVersion, const TMap<FString,FString>& FilesToBeCached );
 
 protected:
 	/** Initializes the PDB Cache. */
@@ -325,9 +341,9 @@ protected:
 	/**
 	 * @return the size of the PDB cache entry, in GBs.
 	 */
-	int32 GetPDBCacheEntrySizeGB( const FString& InLabel ) const
+	int32 GetPDBCacheEntrySizeGB( const FString& PathOrLabel ) const
 	{
-		return PDBCacheEntries.FindChecked( InLabel )->SizeGB;
+		return PDBCacheEntries.FindChecked( PathOrLabel )->SizeGB;
 	}
 
 	/**
@@ -360,7 +376,7 @@ protected:
 	/**
 	 *	Reads an existing PDB Cache entry.
 	 */
-	FPDBCacheEntryRef ReadPDBCacheEntry( const FString& InLabel );
+	FPDBCacheEntryRef ReadPDBCacheEntry( const FString& Directory );
 
 	/**
 	 *	Sort PDB Cache entries by last access time.
@@ -374,12 +390,13 @@ protected:
 	 *	Removes a PDB Cache entry from the database.
 	 *	Also removes all external files associated with this PDB Cache entry.
 	 */
-	void RemovePDBCacheEntry( const FString& InLabel );
+	void RemovePDBCacheEntry( const FString& Directory );
 
-	/** Replaces all / chars with _ for the specified label name. */
-	FString CleanLabelName( const FString& InLabel ) const
+	/** Replaces all invalid chars with + for the specified name. */
+	FString EscapePath( const FString& PathOrLabel ) const
 	{
-		return InLabel.Replace( TEXT( "/" ), TEXT( "_" ) );
+		// @see AutomationTool.CommandUtils.EscapePath 
+		return PathOrLabel.Replace( TEXT( ":" ), TEXT( "" ) ).Replace( TEXT( "/" ), TEXT( "+" ) ).Replace( TEXT( "\\" ), TEXT( "+" ) ).Replace( TEXT( " " ), TEXT( "+" ) );
 	}
 };
 
@@ -389,12 +406,27 @@ class CRASHDEBUGHELPER_API ICrashDebugHelper
 protected:
 	/** The depot name that the handler is using to sync from */
 	FString DepotName;
+
+	/** Depot root, only used in conjunction with the PDB Cache. */
+	FString DepotRoot;
+
 	/** The local folder where symbol files should be stored */
 	FString LocalSymbolStore;
-	/** In the event that a database query for a label fails, use this pattern to search in source control for the label */
+
+	/**
+	 *	Pattern to search in source control for the label.
+	 *	This somehow works for older crashes, before 4.2 and for the newest one,
+	 *	bur for the newest we also need to look for the executables on the network drive.
+	 *	This may change in future.
+	 */
 	FString SourceControlBuildLabelPattern;
-	/** The builtFromCL to use instead of extracting from the minidump */
-	int32 BuiltFromCL;
+	
+	/**
+	 * Pattern to search in the network driver for the executables.
+	 * Valid from 4.2 UE builds. (CL-2068994)
+	 * This may change in future.
+	 */
+	FString ExecutablePathPattern;
 	
 	/** Indicates that the crash handler is ready to do work */
 	bool bInitialized;
@@ -417,18 +449,6 @@ public:
 	 */
 	virtual bool Init();
 
-	/** Get the depot name */
-	const FString& GetDepotName() const
-	{
-		return DepotName;
-	}
-
-	/** Set the depot name */
-	void SetDepotName(const FString& InDepotName)
-	{
-		DepotName = InDepotName;
-	}
-
 	/**
 	 *	Parse the given crash dump, determining EngineVersion of the build that produced it - if possible. 
 	 *
@@ -454,18 +474,6 @@ public:
 		return false;
 	}
 
-	/**
-	 *	Get the source control label for the given engine version
-	 *
-	 *	@param	InEngineVersion		The engine version to retrieve the label for
-	 *
-	 *	@return	FString				The label containing that engine version build, empty if not found
-	 */
-	virtual FString GetLabelFromEngineVersion(int32 InEngineVersion)
-	{
-		return RetrieveBuildLabel(InEngineVersion, -1);
-	}
-
 	/**
 	 *	Get the source control label for the given changelist number
 	 *
@@ -475,7 +483,7 @@ public:
 	 */
 	virtual FString GetLabelFromChangelistNumber(int32 InChangelistNumber)
 	{
-		return RetrieveBuildLabel(-1, InChangelistNumber);
+		return RetrieveBuildLabel(InChangelistNumber);
 	}
 
 	/**
@@ -516,45 +524,42 @@ public:
 	/**
 	 *	Sync the binaries and associated pdbs to the requested label.
 	 *
-	 *	@params	InLabel			The name of the crash dump file to debug
-	 *	@params	InModuleNames	An array of modules to sync
-	 *
 	 *	@return	bool		true if successful, false if not
 	 */
-	virtual bool SyncModules( FCrashInfo* CrashInfo );
+	virtual bool SyncModules();
 
 	/**
 	 *	Sync a single source file to the requested label.
 	 *
-	 *	@params	CrashInfo	A description of a crash
-	 *
 	 *	@return	bool		true if successful, false if not
 	 */
-	virtual bool SyncSourceFile( FCrashInfo* CrashInfo );
+	virtual bool SyncSourceFile();
 
 	/**
 	 *	Extract lines from a source file, and add to the crash report.
-	 *
-	 *	@params	CrashInfo	A description of a crash
 	 */
-	virtual void AddSourceToReport( FCrashInfo* CrashInfo );
+	virtual void AddSourceToReport();
 
 	/**
 	 *	Extract annotated lines from a source file stored in Perforce, and add to the crash report.
-	 *
-	 *	@params	CrashInfo	A description of a crash
 	 */
-	virtual bool AddAnnotatedSourceToReport( FCrashInfo* CrashInfo );
+	virtual bool AddAnnotatedSourceToReport();
 
 protected:
 	/**
-	 *	Retrieve the build label for the given engine version or changelist number.
+	 *	Retrieve the build label for the given changelist number.
 	 *
-	 *	@param	InEngineVersion		The engine version to retrieve the label for. If -1, then will use InChangelistNumber
 	 *	@param	InChangelistNumber	The changelist number to retrieve the label for
 	 *	@return	FString				The label if successful, empty if it wasn't found or another error
 	 */
-	virtual FString RetrieveBuildLabel(int32 InEngineVersion, int32 InChangelistNumber);
+	FString RetrieveBuildLabel(int32 InChangelistNumber);
+
+	/**
+	 *	Retrieve the build label and the network path for the given changelist number.
+	 *
+	 *	@param	InChangelistNumber	The changelist number to retrieve the label for
+	 */
+	void RetrieveBuildLabelAndNetworkPath( int32 InChangelistNumber );
 
 	bool ReadSourceFile( const TCHAR* InFilename, TArray<FString>& OutStrings );
 
