CrashDebugHelper - More logs regarding the PDB cache

diff --git a/Engine/Source/Developer/CrashDebugHelper/Private/CrashDebugHelper.cpp b/Engine/Source/Developer/CrashDebugHelper/Private/CrashDebugHelper.cpp
--- a/Engine/Source/Developer/CrashDebugHelper/Private/CrashDebugHelper.cpp
+++ b/Engine/Source/Developer/CrashDebugHelper/Private/CrashDebugHelper.cpp
@@ -1002,7 +1002,7 @@ void FPDBCache::CleanPDBCache( int32 DaysToDelete, int32 NumberOfGBsToBeCleaned
 	}
 
 	const double TotalTime = FPlatformTime::Seconds() - StartTime;
-	UE_LOG( LogCrashDebugHelper, Log, TEXT( "PDB Cache cleaned in %.2f ms" ), TotalTime*1000.0f );
+	UE_LOG( LogCrashDebugHelper, Log, TEXT( "PDB Cache cleaned %i GBs in %.2f ms" ), NumGBsCleaned, TotalTime*1000.0f );
 }
 
 FPDBCacheEntryRef FPDBCache::CreateAndAddPDBCacheEntry( const FString& OriginalLabelName, const FString& DepotRoot, const FString& DepotName, const TArray<FString>& FilesToBeCached )
