CrashReport - Tweak to make the debugging easier

diff --git a/Engine/Source/Developer/CrashDebugHelper/Private/CrashDebugHelper.cpp b/Engine/Source/Developer/CrashDebugHelper/Private/CrashDebugHelper.cpp
--- a/Engine/Source/Developer/CrashDebugHelper/Private/CrashDebugHelper.cpp
+++ b/Engine/Source/Developer/CrashDebugHelper/Private/CrashDebugHelper.cpp
@@ -766,13 +766,14 @@ FString ICrashDebugHelper::RetrieveBuildLabel(int32 InChangelistNumber)
 		TArray< TSharedRef<ISourceControlLabel> > Labels = ISourceControlModule::Get().GetProvider().GetLabels( TestLabel );
 		if ( Labels.Num() > 0 )
 		{
+			const int32 LabelIndex = 0;
 			// If we found more than one label, warn about it and just use the first one
 			if ( Labels.Num() > 1 )
 			{
-				UE_LOG(LogCrashDebugHelper, Warning, TEXT("RetrieveBuildLabel: More than one build label found with pattern %s - Using label %s"), *TestLabel, *Labels[0]->GetName());
+				UE_LOG(LogCrashDebugHelper, Warning, TEXT("RetrieveBuildLabel: More than one build label found with pattern %s - Using label %s"), *TestLabel, *Labels[LabelIndex]->GetName());
 			}
 
-			FoundLabelString = Labels[0]->GetName();
+			FoundLabelString = Labels[LabelIndex]->GetName();
 			UE_LOG(LogCrashDebugHelper, Log, TEXT("RetrieveBuildLabel: Found label %s matching pattern %s in source control."), *FoundLabelString, *TestLabel);
 		}
 	}
@@ -797,13 +798,14 @@ void ICrashDebugHelper::RetrieveBuildLabelAndNetworkPath( int32 InChangelistNumb
 		TArray< TSharedRef<ISourceControlLabel> > Labels = ISourceControlModule::Get().GetProvider().GetLabels( TestLabelWithCL );
 		if( Labels.Num() > 0 )
 		{
+			const int32 LabelIndex = 0;
 			// If we found more than one label, warn about it and just use the first one
 			if( Labels.Num() > 1 )
 			{
-				UE_LOG( LogCrashDebugHelper, Warning, TEXT( "RetrieveBuildLabelAndNetworkPath: More than one build label found with pattern '%s' - Using label '%s'" ), *TestLabelWithCL, *Labels[0]->GetName() );
+				UE_LOG( LogCrashDebugHelper, Warning, TEXT( "RetrieveBuildLabelAndNetworkPath: More than one build label found with pattern '%s' - Using label '%s'" ), *TestLabelWithCL, *Labels[LabelIndex]->GetName() );
 			}
 
-			CrashInfo.LabelName = Labels[0]->GetName();
+			CrashInfo.LabelName = Labels[LabelIndex]->GetName();
 			UE_LOG( LogCrashDebugHelper, Log, TEXT( "RetrieveBuildLabelAndNetworkPath: Found label '%s' matching pattern '%s' in source control." ), *CrashInfo.LabelName, *TestLabelWithCL );
 		}
 	}
