CrashReport - Optimized gathering the paths for the PDB and binaries for locally made crashes

diff --git a/Engine/Source/Developer/CrashDebugHelper/Private/Windows/WindowsPlatformStackWalkExt.cpp b/Engine/Source/Developer/CrashDebugHelper/Private/Windows/WindowsPlatformStackWalkExt.cpp
--- a/Engine/Source/Developer/CrashDebugHelper/Private/Windows/WindowsPlatformStackWalkExt.cpp
+++ b/Engine/Source/Developer/CrashDebugHelper/Private/Windows/WindowsPlatformStackWalkExt.cpp
@@ -168,6 +168,97 @@ void FWindowsPlatformStackWalkExt::GetExeFileVersionAndModuleList( FCrashModuleI
 	}
 }
 
+/**
+ *	Optimized version used to grab the list of the PDB and binary files.
+ *	Based on the FPackageDependencyTimestampVisitor.
+ */
+class FSymbolAndExecutableFinder : public IPlatformFile::FDirectoryVisitor
+{
+private:
+	/** The file interface to use for any file operations. */
+	IPlatformFile& FileInterface;
+
+	/**
+	 *	A list of file types that should be included in the file paths.
+	 *	An extension with a dot.
+	 */
+	TArray<FString> FiletypesToInclude;
+
+	/** A list of directories that we should not traverse into. */
+	TArray<FString> DirectoriesToIgnore;
+
+	/** Relative paths to local files. */
+	TArray<FString> FilePaths;
+
+public:
+
+	FSymbolAndExecutableFinder( IPlatformFile& InFileInterface, const TArray<FString>& InFiletypesToInclude, const TArray<FString>& InDirectoriesToIgnore )
+		: FileInterface( InFileInterface )
+		, FiletypesToInclude( InFiletypesToInclude )
+		, DirectoriesToIgnore( InDirectoriesToIgnore )
+	{}
+
+	virtual bool Visit( const TCHAR* FilenameOrDirectory, bool bIsDirectory )
+	{
+		// Make sure all paths are "standardized" so the other end can match up with it's own standardized paths
+		FString RelativeFilename = FilenameOrDirectory;
+		FPaths::MakeStandardFilename( RelativeFilename );
+
+		// Grab the file.
+		if( !bIsDirectory )
+		{
+			if( FiletypesToInclude.Num() > 0 )
+			{
+				// If it is a file that matches the filetypes, add it.
+				for( const auto& FileExtension : FiletypesToInclude )
+				{
+					if( RelativeFilename.Contains( FileExtension ) )
+					{
+						new(FilePaths) FString( RelativeFilename );
+						break;
+					}
+				}
+			}
+			else
+			{
+				new(FilePaths) FString( RelativeFilename );
+			}
+		}
+
+		// Iterate over directories we care about.
+		if( bIsDirectory )
+		{
+			bool bShouldRecurse = true;
+			// Look in all the ignore directories looking for a match.
+			for( const auto& Directory : DirectoriesToIgnore )
+			{
+				if( RelativeFilename.Contains( Directory ) )
+				{
+					bShouldRecurse = false;
+					break;
+				}
+			}
+
+			// Recurse if we should.
+			if( bShouldRecurse )
+			{
+				FileInterface.IterateDirectory( FilenameOrDirectory, *this );
+			}
+		}
+
+		return true;
+	}
+
+	
+	/**
+	 * @return a reference to the results.
+	 */
+	const TArray<FString>& GetFilePaths() const
+	{
+		return FilePaths;
+	}
+};
+
 /** 
  * Set the symbol paths based on the module paths
  */
@@ -202,14 +293,26 @@ void FWindowsPlatformStackWalkExt::SetSymbolPathsFromModules()
 	// For locally launched minidump diagnostics.
 	else
 	{
-		// FindFilesRecursive for root is quite slow.
-		TArray<FString> ModulesAndSymbols;
-		IFileManager::Get().FindFilesRecursive( ModulesAndSymbols, *FPaths::RootDir(), TEXT( "*.pdb" ), true, false, false );
-		IFileManager::Get().FindFilesRecursive( ModulesAndSymbols, *FPaths::RootDir(), TEXT( "*.dll" ), true, false, false );
-	
+		// We care only about the PDB and binary files.
+		TArray<FString> FiletypesToInclude;
+		FiletypesToInclude.Add( TEXT( ".pdb" ) );
+		FiletypesToInclude.Add( TEXT( ".dll" ) );
+		FiletypesToInclude.Add( TEXT( ".exe" ) );
+		// Content folders don't contains any binary files, so skip them to improve the finder performance.
+		TArray<FString> DirectoriesToIgnore;
+		DirectoriesToIgnore.Add( TEXT( "/Content/" ) );
+		DirectoriesToIgnore.Add( TEXT( "/DerivedDataCache/" ) );
+		DirectoriesToIgnore.Add( TEXT( "/Intermediate/" ) );
+		DirectoriesToIgnore.Add( TEXT( "/Documentation/" ) );
+		DirectoriesToIgnore.Add( TEXT( "/NoRedist/" ) );
+
+
+		FSymbolAndExecutableFinder SymbolAndExecutableFinder( FPlatformFileManager::Get().GetPlatformFile(), FiletypesToInclude, DirectoriesToIgnore );
+		SymbolAndExecutableFinder.Visit( *FPaths::RootDir(), true );
+
 
 		TSet<FString> SymbolPaths;
-		for( const auto& Filename : ModulesAndSymbols )
+		for( const auto& Filename : SymbolAndExecutableFinder.GetFilePaths() )
 		{
 			const FString Path = FPaths::GetPath( Filename );
 			if( Path.Len() > 0 )
