Fixed spelling

diff --git a/Engine/Source/Developer/CrashDebugHelper/Private/Windows/WindowsPlatformStackWalkExt.cpp b/Engine/Source/Developer/CrashDebugHelper/Private/Windows/WindowsPlatformStackWalkExt.cpp
--- a/Engine/Source/Developer/CrashDebugHelper/Private/Windows/WindowsPlatformStackWalkExt.cpp
+++ b/Engine/Source/Developer/CrashDebugHelper/Private/Windows/WindowsPlatformStackWalkExt.cpp
@@ -298,7 +298,7 @@ void FWindowsPlatformStackWalkExt::SetSymbolPathsFromModules()
 		FiletypesToInclude.Add( TEXT( ".pdb" ) );
 		FiletypesToInclude.Add( TEXT( ".dll" ) );
 		FiletypesToInclude.Add( TEXT( ".exe" ) );
-		// Content folders don't contains any binary files, so skip them to improve the finder performance.
+		// These folders don't contain any binary files, so skip them to improve the finder performance.
 		TArray<FString> DirectoriesToIgnore;
 		DirectoriesToIgnore.Add( TEXT( "/Content/" ) );
 		DirectoriesToIgnore.Add( TEXT( "/DerivedDataCache/" ) );
