CrashReport - PDB Cache stores the metadata for each entry to avoid time consuming FindFilesRecursive during initialization

diff --git a/Engine/Source/Developer/CrashDebugHelper/Private/CrashDebugHelper.cpp b/Engine/Source/Developer/CrashDebugHelper/Private/CrashDebugHelper.cpp
--- a/Engine/Source/Developer/CrashDebugHelper/Private/CrashDebugHelper.cpp
+++ b/Engine/Source/Developer/CrashDebugHelper/Private/CrashDebugHelper.cpp
@@ -844,7 +844,8 @@ void ICrashDebugHelper::RetrieveBuildLabelAndNetworkPath( int32 InChangelistNumb
 	PDB Cache implementation
 -----------------------------------------------------------------------------*/
 
-const TCHAR* FPDBCache::PDBTimeStampFile = TEXT( "PDBTimeStamp.txt" );
+const TCHAR* FPDBCache::PDBTimeStampFileNoMeta = TEXT( "PDBTimeStamp.txt" );
+const TCHAR* FPDBCache::PDBTimeStampFile = TEXT( "PDBTimeStamp.bin" );
 
 void FPDBCache::Init()
 {
@@ -1013,11 +1014,6 @@ FPDBCacheEntryRef FPDBCache::CreateAndAddPDBCacheEntry( const FString& OriginalL
 	const FString EntryDirectory = PDBCachePath / CleanedLabelName;
 	const FString EntryTimeStampFilename = EntryDirectory / PDBTimeStampFile;
 
-	// Verify there is an entry timestamp file.
-	const bool bSaved = FFileHelper::SaveStringToFile( EntryTimeStampFilename, *EntryTimeStampFilename );
-	UE_CLOG( !bSaved, LogCrashDebugHelper, Fatal, TEXT( "Couldn't save the timestamp file to %s" ), *EntryTimeStampFilename );
-	const FDateTime LastAccessTime = IFileManager::Get().GetTimeStamp( *EntryTimeStampFilename );
-
 	const FString LocalDepotDir = DepotRoot / DepotName.Replace( P4_DEPOT_PREFIX, TEXT( "" ) );
 
 	UE_LOG( LogCrashDebugHelper, Warning, TEXT( "PDB Cache entry '%s' is being copied from '%s', it will take some time" ), *CleanedLabelName, *OriginalLabelName );
@@ -1049,11 +1045,13 @@ FPDBCacheEntryRef FPDBCache::CreateAndAddPDBCacheEntry( const FString& OriginalL
 	}
 
 	// Round-up the size.
-	const int32 SizeGB = (int32)FMath::DivideAndRoundUp( TotalSize, (int64)NUM_BYTES_PER_GB );
+	int32 SizeGB = (int32)FMath::DivideAndRoundUp( TotalSize, (int64)NUM_BYTES_PER_GB );
+	FPDBCacheEntryRef NewCacheEntry = MakeShareable( new FPDBCacheEntry( CachedFiles, CleanedLabelName, FDateTime::Now(), SizeGB ) );
 
-	FPDBCacheEntryRef NewCacheEntry = MakeShareable( new FPDBCacheEntry( CleanedLabelName, SizeGB ) );
-	NewCacheEntry->LastAccessTime = LastAccessTime;
-	NewCacheEntry->Files = CachedFiles;
+	// Verify there is an entry timestamp file, write the size of a PDB cache to avoid time consuming FindFilesRecursive during initialization.
+	TAutoPtr<FArchive> FileWriter( IFileManager::Get().CreateFileWriter( *EntryTimeStampFilename ) );
+	UE_CLOG( !FileWriter, LogCrashDebugHelper, Fatal, TEXT( "Couldn't save the timestamp file to %s" ), *EntryTimeStampFilename );
+	*FileWriter << *NewCacheEntry;
 	
 	PDBCacheEntries.Add( CleanedLabelName, NewCacheEntry );
 	SortPDBCache();
@@ -1067,11 +1065,6 @@ FPDBCacheEntryRef FPDBCache::CreateAndAddPDBCacheEntryMixed( const FString& Prod
 	const FString EntryDirectory = PDBCachePath / ProductVersion;
 	const FString EntryTimeStampFilename = EntryDirectory / PDBTimeStampFile;
  
-	// Verify there is an entry timestamp file.
-	const bool bSaved = FFileHelper::SaveStringToFile( EntryTimeStampFilename, *EntryTimeStampFilename );
-	UE_CLOG( !bSaved, LogCrashDebugHelper, Fatal, TEXT( "Couldn't save the timestamp file to %s" ), *EntryTimeStampFilename );
-	const FDateTime LastAccessTime = IFileManager::Get().GetTimeStamp( *EntryTimeStampFilename );
- 
 	UE_LOG( LogCrashDebugHelper, Warning, TEXT( "PDB Cache entry '%s' is being created from %i files, it will take some time" ), *ProductVersion, FilesToBeCached.Num() );
 	for( const auto& It : FilesToBeCached )
 	{
@@ -1092,11 +1085,13 @@ FPDBCacheEntryRef FPDBCache::CreateAndAddPDBCacheEntryMixed( const FString& Prod
 	}
  
 	// Round-up the size.
-	const int32 SizeGB = (int32)FMath::DivideAndRoundUp( TotalSize, (int64)NUM_BYTES_PER_GB );
- 
-	FPDBCacheEntryRef NewCacheEntry = MakeShareable( new FPDBCacheEntry( ProductVersion, SizeGB ) );
-	NewCacheEntry->LastAccessTime = LastAccessTime;
-	NewCacheEntry->Files = CachedFiles;
+	int32 SizeGB = (int32)FMath::DivideAndRoundUp( TotalSize, (int64)NUM_BYTES_PER_GB );
+	FPDBCacheEntryRef NewCacheEntry = MakeShareable( new FPDBCacheEntry( CachedFiles, ProductVersion, FDateTime::Now(), SizeGB ) );
+
+	// Verify there is an entry timestamp file, write the size of a PDB cache to avoid time consuming FindFilesRecursive during initialization.
+	TAutoPtr<FArchive> FileWriter( IFileManager::Get().CreateFileWriter( *EntryTimeStampFilename ) );
+	UE_CLOG( !FileWriter, LogCrashDebugHelper, Fatal, TEXT( "Couldn't save the timestamp file to %s" ), *EntryTimeStampFilename );
+	*FileWriter << *NewCacheEntry;
  
 	PDBCacheEntries.Add( ProductVersion, NewCacheEntry );
 	SortPDBCache();
@@ -1107,29 +1102,53 @@ FPDBCacheEntryRef FPDBCache::CreateAndAddPDBCacheEntryMixed( const FString& Prod
 FPDBCacheEntryRef FPDBCache::ReadPDBCacheEntry( const FString& Directory )
 {
 	const FString EntryDirectory = PDBCachePath / Directory;
+	const FString EntryTimeStampFilenameNoMeta = EntryDirectory / PDBTimeStampFileNoMeta;
 	const FString EntryTimeStampFilename = EntryDirectory / PDBTimeStampFile;
 
 	// Verify there is an entry timestamp file.
+	const FDateTime LastAccessTimeNoMeta = IFileManager::Get().GetTimeStamp( *EntryTimeStampFilenameNoMeta );
 	const FDateTime LastAccessTime = IFileManager::Get().GetTimeStamp( *EntryTimeStampFilename );
 
-	TArray<FString> PDBFiles;
-	IFileManager::Get().FindFilesRecursive( PDBFiles, *EntryDirectory, TEXT( "*.*" ), true, false );
+	FPDBCacheEntryPtr NewEntry;
 
-	// Calculate the size of this PDB Cache entry.
-	int64 TotalSize = 0;
-	for( const auto& Filename : PDBFiles )
+	if( LastAccessTime != FDateTime::MinValue() )
 	{
-		const int64 FileSize = IFileManager::Get().FileSize( *Filename );
-		TotalSize += FileSize;
+		// Read the metadata
+		TAutoPtr<FArchive> FileReader( IFileManager::Get().CreateFileReader( *EntryTimeStampFilename ) );
+		NewEntry = MakeShareable( new FPDBCacheEntry( LastAccessTime ) );
+		*FileReader << *NewEntry;
 	}
+	else if( LastAccessTimeNoMeta != FDateTime::MinValue() )
+	{
+		// Calculate the size of this PDB Cache entry and update to the new version.
+		TArray<FString> PDBFiles;
+		IFileManager::Get().FindFilesRecursive( PDBFiles, *EntryDirectory, TEXT( "*.*" ), true, false );
 
-	// Round-up the size.
-	const int32 SizeGB = (int32)FMath::DivideAndRoundUp( TotalSize, (int64)NUM_BYTES_PER_GB );
+		// Calculate the size of this PDB Cache entry.
+		int64 TotalSize = 0;
+		for( const auto& Filename : PDBFiles )
+		{
+			const int64 FileSize = IFileManager::Get().FileSize( *Filename );
+			TotalSize += FileSize;
+		}
+
+		// Round-up the size.
+		const int32 SizeGB = (int32)FMath::DivideAndRoundUp( TotalSize, (int64)NUM_BYTES_PER_GB );
+		NewEntry = MakeShareable( new FPDBCacheEntry( PDBFiles, Directory, LastAccessTimeNoMeta, SizeGB ) );
 
-	FPDBCacheEntryRef NewEntry = MakeShareable( new FPDBCacheEntry( Directory, SizeGB ) );
-	NewEntry->LastAccessTime = LastAccessTime;
-	NewEntry->Files = PDBFiles;
-	return NewEntry;
+		// Remove the old file and save the metadata.
+		TAutoPtr<FArchive> FileWriter( IFileManager::Get().CreateFileWriter( *EntryTimeStampFilename ) );
+		*FileWriter << *NewEntry;
+
+		IFileManager::Get().Delete( *EntryTimeStampFilenameNoMeta );
+	}
+	else
+	{
+		// Something wrong.
+		check( 0 );
+	}
+
+	return NewEntry.ToSharedRef();
 }
 
 void FPDBCache::TouchPDBCacheEntry( const FString& Directory )
@@ -1164,6 +1183,3 @@ FPDBCacheEntryRef FPDBCache::FindAndTouchPDBCacheEntry( const FString& PathOrLab
 	TouchPDBCacheEntry( CacheEntry->Directory );
 	return CacheEntry;
 }
-
-
-
diff --git a/Engine/Source/Developer/CrashDebugHelper/Public/CrashDebugHelper.h b/Engine/Source/Developer/CrashDebugHelper/Public/CrashDebugHelper.h
--- a/Engine/Source/Developer/CrashDebugHelper/Public/CrashDebugHelper.h
+++ b/Engine/Source/Developer/CrashDebugHelper/Public/CrashDebugHelper.h
@@ -205,15 +205,23 @@ struct CRASHDEBUGHELPER_API FCrashDebugInfo
 /** Helper struct that holds various information about one PDB Cache entry. */
 struct FPDBCacheEntry
 {
+	/** Default constructor. */
+	FPDBCacheEntry( const FDateTime InLastAccessTime )
+		: SizeGB( 0 )
+		, LastAccessTime( InLastAccessTime )
+	{}
+
 	/** Initialization constructor. */
-	FPDBCacheEntry( const FString& InDirectory, const int32 InSizeGB )
-		: Directory( InDirectory )
+	FPDBCacheEntry( const TArray<FString>& InFiles, const FString& InDirectory, const FDateTime InLastAccessTime, const int32 InSizeGB )
+		: Files( InFiles )
+		, Directory( InDirectory )
+		, LastAccessTime( InLastAccessTime )
 		, SizeGB( InSizeGB )
 	{}
 
 	void SetLastAccessTimeToNow()
 	{
-		LastAccessTime = FDateTime::UtcNow();
+		LastAccessTime = FDateTime::Now();
 	}
 
 	/** Paths to files associated with this PDB Cache entry. */
@@ -227,6 +235,14 @@ struct FPDBCacheEntry
 
 	/** Size of the cache entry, in GBs. Rounded-up. */
 	const int32 SizeGB;
+
+	/**
+	 * Serializer.
+	 */
+	friend FArchive& operator<<(FArchive& Ar, FPDBCacheEntry& Entry)
+	{
+		return Ar << Entry.Files << (FString&)Entry.Directory << (int32&)Entry.SizeGB;
+	}
 };
 
 struct FPDBCacheEntryByAccessTime
@@ -262,6 +278,9 @@ protected:
 	};
 
 	/** Dummy file used to read/set the file timestamp. */
+	static const TCHAR* PDBTimeStampFileNoMeta;
+
+	/** Data file used to read/set the file timestamp, contains all metadata. */
 	static const TCHAR* PDBTimeStampFile;
 
 	/** Map of the PDB Cache entries. */
