CIS fix

diff --git a/Engine/Source/Developer/CrashDebugHelper/Public/CrashDebugHelper.h b/Engine/Source/Developer/CrashDebugHelper/Public/CrashDebugHelper.h
--- a/Engine/Source/Developer/CrashDebugHelper/Public/CrashDebugHelper.h
+++ b/Engine/Source/Developer/CrashDebugHelper/Public/CrashDebugHelper.h
@@ -228,7 +228,7 @@ struct FPDBCacheEntry
 	TArray<FString> Files;
 
 	/** The path associated with this PDB Cache entry. */
-	const FString Directory;
+	FString Directory;
 	
 	/** Last access time, changed every time this PDB cache entry is used. */
 	FDateTime LastAccessTime;
