CrashReport - Sync modules and PDB Cache should work even if the network build has been removed, but a PDB cache entry is on the local drive

diff --git a/Engine/Source/Developer/CrashDebugHelper/Private/CrashDebugHelper.cpp b/Engine/Source/Developer/CrashDebugHelper/Private/CrashDebugHelper.cpp
--- a/Engine/Source/Developer/CrashDebugHelper/Private/CrashDebugHelper.cpp
+++ b/Engine/Source/Developer/CrashDebugHelper/Private/CrashDebugHelper.cpp
@@ -159,175 +159,171 @@ bool ICrashDebugHelper::SyncModules()
 		TSharedRef<ISourceControlLabel> Label = Labels[0];
 		TSet<FString> FilesToSync;
 
-		if( bUseNetworkPathExecutables )
+		// Use product version instead of label name to make a distinguish between chosen methods.
+		const bool bContainsProductVersion = PDBCache.UsePDBCache() && PDBCache.ContainsPDBCacheEntry( CrashInfo.ProductVersion );
+		const bool bContainsLabelName = PDBCache.UsePDBCache() && PDBCache.ContainsPDBCacheEntry( CrashInfo.LabelName );
+
+		if( bContainsProductVersion )
 		{
-			// Use product version instead of label name to make a distinguish between chosen methods.
-			const bool bContains = PDBCache.UsePDBCache() && PDBCache.ContainsPDBCacheEntry( CrashInfo.ProductVersion );
-			if( bContains )
-			{
-				UE_LOG( LogCrashDebugHelper, Warning, TEXT( "Label %s found in the PDB Cache, using it" ), *CrashInfo.ProductVersion );
-				CrashInfo.PDBCacheEntry = PDBCache.FindAndTouchPDBCacheEntry( CrashInfo.ProductVersion );
-				return true;
-			}
-			else
-			{
-				SCOPE_LOG_TIME_IN_SECONDS( TEXT( "SyncModulesAndNetwork" ), nullptr );
+			UE_LOG( LogCrashDebugHelper, Warning, TEXT( "Label %s found in the PDB Cache, using it" ), *CrashInfo.ProductVersion );
+			CrashInfo.PDBCacheEntry = PDBCache.FindAndTouchPDBCacheEntry( CrashInfo.ProductVersion );
+			return true;
+		}
+		else if( bContainsLabelName )
+		{
+			UE_LOG( LogCrashDebugHelper, Warning, TEXT( "Label '%s' found in the PDB Cache, using it" ), *CrashInfo.LabelName );
+			CrashInfo.PDBCacheEntry = PDBCache.FindAndTouchPDBCacheEntry( CrashInfo.LabelName );
+			return true;
+		}
+		else if( bUseNetworkPathExecutables )
+		{			
+			SCOPE_LOG_TIME_IN_SECONDS( TEXT( "SyncModulesAndNetwork" ), nullptr );
 
-				// This is a bit hacky, probably will be valid until next build system change.
-				const FString ProductNetworkPath = ExecutablePathPattern.Replace( TEXT( "%PRODUCT_VERSION%" ), *CrashInfo.ProductVersion );
+			// This is a bit hacky, probably will be valid until next build system change.
+			const FString ProductNetworkPath = ExecutablePathPattern.Replace( TEXT( "%PRODUCT_VERSION%" ), *CrashInfo.ProductVersion );
 
-				// Grab information about symbols.
-				TArray< TSharedRef<class ISourceControlRevision, ESPMode::ThreadSafe> > PDBSourceControlRevisions;
-				const FString PDBsPath = FString::Printf( TEXT( "%s/%s....pdb" ), *DepotName, UESymbols );
-				Label->GetFileRevisions( PDBsPath, PDBSourceControlRevisions );
-				
-				TSet<FString> PDBPaths;
-				for( const auto& PDBSrc : PDBSourceControlRevisions )
-				{
-					PDBPaths.Add( PDBSrc->GetFilename() );
-				}
+			// Grab information about symbols.
+			TArray< TSharedRef<class ISourceControlRevision, ESPMode::ThreadSafe> > PDBSourceControlRevisions;
+			const FString PDBsPath = FString::Printf( TEXT( "%s/%s....pdb" ), *DepotName, UESymbols );
+			Label->GetFileRevisions( PDBsPath, PDBSourceControlRevisions );
 
-				// Now, sync symbols.
-				for( const auto& PDBPath : PDBPaths )
+			TSet<FString> PDBPaths;
+			for( const auto& PDBSrc : PDBSourceControlRevisions )
+			{
+				PDBPaths.Add( PDBSrc->GetFilename() );
+			}
+
+			// Now, sync symbols.
+			for( const auto& PDBPath : PDBPaths )
+			{
+				if( Label->Sync( PDBPath ) )
 				{
-					if( Label->Sync( PDBPath ) )
-					{
-						UE_LOG( LogCrashDebugHelper, Warning, TEXT( "Synced PDB '%s'" ), *PDBPath );
-					}
+					UE_LOG( LogCrashDebugHelper, Warning, TEXT( "Synced PDB '%s'" ), *PDBPath );
 				}
+			}
 
-				// Find all the executables in the product network path.
-				TArray<FString> NetworkExecutables;
-				IFileManager::Get().FindFilesRecursive( NetworkExecutables, *ProductNetworkPath, TEXT( "*.dll" ), true, false, false );
-				IFileManager::Get().FindFilesRecursive( NetworkExecutables, *ProductNetworkPath, TEXT( "*.exe" ), true, false, false );
+			// Find all the executables in the product network path.
+			TArray<FString> NetworkExecutables;
+			IFileManager::Get().FindFilesRecursive( NetworkExecutables, *ProductNetworkPath, TEXT( "*.dll" ), true, false, false );
+			IFileManager::Get().FindFilesRecursive( NetworkExecutables, *ProductNetworkPath, TEXT( "*.exe" ), true, false, false );
 
-				// From=Full pathname
-				// To=Relative pathname
-				TMap<FString, FString> FilesToBeCached;
+			// From=Full pathname
+			// To=Relative pathname
+			TMap<FString, FString> FilesToBeCached;
 
-				// If a symbol matches an executable, add the pair to the list of files that should be cached.
-				for( const auto& NetworkExecutableFullpath : NetworkExecutables )
+			// If a symbol matches an executable, add the pair to the list of files that should be cached.
+			for( const auto& NetworkExecutableFullpath : NetworkExecutables )
+			{
+				for( const auto& PDBPath : PDBPaths )
 				{
-					for( const auto& PDBPath : PDBPaths )
-					{
-						const FString PDBRelativePath = PDBPath.Replace( *DepotName, TEXT( "" ) ).Replace( UESymbols, TEXT( "" ) );
-						const FString PDBFullpath = DepotRoot / PDBPath.Replace( P4_DEPOT_PREFIX, TEXT( "" ) );
+					const FString PDBRelativePath = PDBPath.Replace( *DepotName, TEXT( "" ) ).Replace( UESymbols, TEXT( "" ) );
+					const FString PDBFullpath = DepotRoot / PDBPath.Replace( P4_DEPOT_PREFIX, TEXT( "" ) );
 
-						const FString PDBMatch = PDBRelativePath.Replace( TEXT( "pdb" ), TEXT( "" ) );
-						const FString NetworkRelativePath = NetworkExecutableFullpath.Replace( *ProductNetworkPath, TEXT( "" ) );
-						const bool bMatch = NetworkExecutableFullpath.Contains( PDBMatch );
-						if( bMatch )
-						{
-							// From -> Where ?
-							FilesToBeCached.Add( NetworkExecutableFullpath, NetworkRelativePath );
-							FilesToBeCached.Add( PDBFullpath, PDBRelativePath );
-							break;
-						}
+					const FString PDBMatch = PDBRelativePath.Replace( TEXT( "pdb" ), TEXT( "" ) );
+					const FString NetworkRelativePath = NetworkExecutableFullpath.Replace( *ProductNetworkPath, TEXT( "" ) );
+					const bool bMatch = NetworkExecutableFullpath.Contains( PDBMatch );
+					if( bMatch )
+					{
+						// From -> Where ?
+						FilesToBeCached.Add( NetworkExecutableFullpath, NetworkRelativePath );
+						FilesToBeCached.Add( PDBFullpath, PDBRelativePath );
+						break;
 					}
 				}
+			}
 
-				if( PDBCache.UsePDBCache() )
-				{
-					// Initialize and add a new PDB Cache entry to the database.
-					CrashInfo.PDBCacheEntry = PDBCache.CreateAndAddPDBCacheEntryMixed( CrashInfo.ProductVersion, FilesToBeCached );
-				}
+			if( PDBCache.UsePDBCache() )
+			{
+				// Initialize and add a new PDB Cache entry to the database.
+				CrashInfo.PDBCacheEntry = PDBCache.CreateAndAddPDBCacheEntryMixed( CrashInfo.ProductVersion, FilesToBeCached );
 			}
 		}
 		else
 		{
 			TArray<FString> FilesToBeCached;
-			const bool bContains = PDBCache.UsePDBCache() && PDBCache.ContainsPDBCacheEntry( CrashInfo.LabelName );
-			if( bContains )
+			
+			//@TODO: MAC: Excluding labels for Mac since we are only syncing windows binaries here...
+			if( Label->GetName().Contains( TEXT( "Mac" ) ) )
 			{
-				UE_LOG( LogCrashDebugHelper, Warning, TEXT( "Label '%s' found in the PDB Cache, using it" ), *CrashInfo.LabelName );
-				CrashInfo.PDBCacheEntry = PDBCache.FindAndTouchPDBCacheEntry( CrashInfo.LabelName );
-				return true;
+				UE_LOG( LogCrashDebugHelper, Log, TEXT( " Skipping Mac label '%s' when syncing modules." ), *Label->GetName() );
 			}
 			else
 			{
-				//@TODO: MAC: Excluding labels for Mac since we are only syncing windows binaries here...
-				if( Label->GetName().Contains( TEXT( "Mac" ) ) )
-				{
-					UE_LOG( LogCrashDebugHelper, Log, TEXT( " Skipping Mac label '%s' when syncing modules." ), *Label->GetName() );
-				}
-				else
-				{
-					// Sync all the dll, exes, and related symbol files
-					UE_LOG( LogCrashDebugHelper, Log, TEXT( " Syncing modules with label '%s'." ), *Label->GetName() );
+				// Sync all the dll, exes, and related symbol files
+				UE_LOG( LogCrashDebugHelper, Log, TEXT( " Syncing modules with label '%s'." ), *Label->GetName() );
 
-					SCOPE_LOG_TIME_IN_SECONDS( TEXT( "SyncModules" ), nullptr );
+				SCOPE_LOG_TIME_IN_SECONDS( TEXT( "SyncModules" ), nullptr );
 
-					// Grab all dll and pdb files for the specified label.
-					TArray< TSharedRef<class ISourceControlRevision, ESPMode::ThreadSafe> > DLLSourceControlRevisions;
-					const FString DLLsPath = FString::Printf( TEXT( "%s/....dll" ), *DepotName );
-					Label->GetFileRevisions( DLLsPath, DLLSourceControlRevisions );
+				// Grab all dll and pdb files for the specified label.
+				TArray< TSharedRef<class ISourceControlRevision, ESPMode::ThreadSafe> > DLLSourceControlRevisions;
+				const FString DLLsPath = FString::Printf( TEXT( "%s/....dll" ), *DepotName );
+				Label->GetFileRevisions( DLLsPath, DLLSourceControlRevisions );
 
-					TArray< TSharedRef<class ISourceControlRevision, ESPMode::ThreadSafe> > EXESourceControlRevisions;
-					const FString EXEsPath = FString::Printf( TEXT( "%s/....exe" ), *DepotName );
-					Label->GetFileRevisions( EXEsPath, EXESourceControlRevisions );
+				TArray< TSharedRef<class ISourceControlRevision, ESPMode::ThreadSafe> > EXESourceControlRevisions;
+				const FString EXEsPath = FString::Printf( TEXT( "%s/....exe" ), *DepotName );
+				Label->GetFileRevisions( EXEsPath, EXESourceControlRevisions );
 
-					TArray< TSharedRef<class ISourceControlRevision, ESPMode::ThreadSafe> > PDBSourceControlRevisions;
-					const FString PDBsPath = FString::Printf( TEXT( "%s/....pdb" ), *DepotName );
-					Label->GetFileRevisions( PDBsPath, PDBSourceControlRevisions );
-
-					TSet<FString> ModulesPaths;
-					for( const auto& DLLSrc : DLLSourceControlRevisions )
-					{
-						ModulesPaths.Add( DLLSrc->GetFilename().Replace( *DepotName, TEXT( "" ) ) );
-					}
-					for( const auto& EXESrc : EXESourceControlRevisions )
-					{
-						ModulesPaths.Add( EXESrc->GetFilename().Replace( *DepotName, TEXT( "" ) ) );
-					}
+				TArray< TSharedRef<class ISourceControlRevision, ESPMode::ThreadSafe> > PDBSourceControlRevisions;
+				const FString PDBsPath = FString::Printf( TEXT( "%s/....pdb" ), *DepotName );
+				Label->GetFileRevisions( PDBsPath, PDBSourceControlRevisions );
 
-					TSet<FString> PDBPaths;
-					for( const auto& PDBSrc : PDBSourceControlRevisions )
-					{
-						PDBPaths.Add( PDBSrc->GetFilename().Replace( *DepotName, TEXT( "" ) ) );
-					}
+				TSet<FString> ModulesPaths;
+				for( const auto& DLLSrc : DLLSourceControlRevisions )
+				{
+					ModulesPaths.Add( DLLSrc->GetFilename().Replace( *DepotName, TEXT( "" ) ) );
+				}
+				for( const auto& EXESrc : EXESourceControlRevisions )
+				{
+					ModulesPaths.Add( EXESrc->GetFilename().Replace( *DepotName, TEXT( "" ) ) );
+				}
 
-					// Iterate through all module and see if we have dll and pdb associated with the module, if so add it to the files to sync.
-					for( const auto& ModuleName : CrashInfo.ModuleNames )
-					{
-						const FString ModuleNamePDB = ModuleName.Replace( TEXT( ".dll" ), TEXT( ".pdb" ) ).Replace( TEXT( ".exe" ), TEXT( ".pdb" ) );
+				TSet<FString> PDBPaths;
+				for( const auto& PDBSrc : PDBSourceControlRevisions )
+				{
+					PDBPaths.Add( PDBSrc->GetFilename().Replace( *DepotName, TEXT( "" ) ) );
+				}
 
-						for( const auto& ModulePath : ModulesPaths )
-						{
-							const bool bContainsModule = ModulePath.Contains( ModuleName );
-							if( bContainsModule )
-							{
-								FilesToSync.Add( ModulePath );
-							}
-						}
+				// Iterate through all module and see if we have dll and pdb associated with the module, if so add it to the files to sync.
+				for( const auto& ModuleName : CrashInfo.ModuleNames )
+				{
+					const FString ModuleNamePDB = ModuleName.Replace( TEXT( ".dll" ), TEXT( ".pdb" ) ).Replace( TEXT( ".exe" ), TEXT( ".pdb" ) );
 
-						for( const auto& PDBPath : PDBPaths )
+					for( const auto& ModulePath : ModulesPaths )
+					{
+						const bool bContainsModule = ModulePath.Contains( ModuleName );
+						if( bContainsModule )
 						{
-							const bool bContainsPDB = PDBPath.Contains( ModuleNamePDB );
-							if( bContainsPDB )
-							{
-								FilesToSync.Add( PDBPath );
-							}
+							FilesToSync.Add( ModulePath );
 						}
 					}
 
-					// Now, sync all files.
-					for( const auto& Filename : FilesToSync )
+					for( const auto& PDBPath : PDBPaths )
 					{
-						const FString DepotPath = DepotName + Filename;
-						if( Label->Sync( DepotPath ) )
+						const bool bContainsPDB = PDBPath.Contains( ModuleNamePDB );
+						if( bContainsPDB )
 						{
-							UE_LOG( LogCrashDebugHelper, Warning, TEXT( " ... synced binary '%s'." ), *DepotPath );
+							FilesToSync.Add( PDBPath );
 						}
-						FilesToBeCached.Add( DepotPath );
 					}
 				}
 
-				if( PDBCache.UsePDBCache() )
+				// Now, sync all files.
+				for( const auto& Filename : FilesToSync )
 				{
-					// Initialize and add a new PDB Cache entry to the database.
-					CrashInfo.PDBCacheEntry = PDBCache.CreateAndAddPDBCacheEntry( CrashInfo.LabelName, DepotRoot, DepotName, FilesToBeCached );
+					const FString DepotPath = DepotName + Filename;
+					if( Label->Sync( DepotPath ) )
+					{
+						UE_LOG( LogCrashDebugHelper, Warning, TEXT( " ... synced binary '%s'." ), *DepotPath );
+					}
+					FilesToBeCached.Add( DepotPath );
 				}
 			}
+
+			if( PDBCache.UsePDBCache() )
+			{
+				// Initialize and add a new PDB Cache entry to the database.
+				CrashInfo.PDBCacheEntry = PDBCache.CreateAndAddPDBCacheEntry( CrashInfo.LabelName, DepotRoot, DepotName, FilesToBeCached );
+			}
 		}
 	}
 	else
diff --git a/Engine/Source/Developer/CrashDebugHelper/Public/CrashDebugHelper.h b/Engine/Source/Developer/CrashDebugHelper/Public/CrashDebugHelper.h
--- a/Engine/Source/Developer/CrashDebugHelper/Public/CrashDebugHelper.h
+++ b/Engine/Source/Developer/CrashDebugHelper/Public/CrashDebugHelper.h
@@ -221,7 +221,7 @@ struct FPDBCacheEntry
 
 	void SetLastAccessTimeToNow()
 	{
-		LastAccessTime = FDateTime::Now();
+		LastAccessTime = FDateTime::UtcNow();
 	}
 
 	/** Paths to files associated with this PDB Cache entry. */
