CrashReport - Fixed a problem with the execution order

diff --git a/Engine/Source/Developer/CrashDebugHelper/Private/CrashDebugHelper.cpp b/Engine/Source/Developer/CrashDebugHelper/Private/CrashDebugHelper.cpp
--- a/Engine/Source/Developer/CrashDebugHelper/Private/CrashDebugHelper.cpp
+++ b/Engine/Source/Developer/CrashDebugHelper/Private/CrashDebugHelper.cpp
@@ -1152,7 +1152,7 @@ void FPDBCache::TouchPDBCacheEntry( const FString& Directory )
 	const FString EntryDirectory = PDBCachePath / Directory;
 	const FString EntryTimeStampFilename = EntryDirectory / PDBTimeStampFile;
 
-	FPDBCacheEntryRef Entry = PDBCacheEntries.FindChecked( Directory );
+	FPDBCacheEntryRef& Entry = PDBCacheEntries.FindChecked( Directory );
 	Entry->SetLastAccessTimeToNow();
 
 	const bool bResult = IFileManager::Get().SetTimeStamp( *EntryTimeStampFilename, Entry->LastAccessTime );
@@ -1167,15 +1167,16 @@ void FPDBCache::RemovePDBCacheEntry( const FString& Directory )
 
 	FPDBCacheEntryRef& Entry = PDBCacheEntries.FindChecked( Directory );
 	IFileManager::Get().DeleteDirectory( *EntryDirectory, true, true );
-	PDBCacheEntries.Remove( Directory );
-
+	
 	const double TotalTime = FPlatformTime::Seconds() - StartTime;
 	UE_LOG( LogCrashDebugHelper, Warning, TEXT( "PDB Cache entry %s removed in %.2f ms, restored %i GBs" ), *Directory, TotalTime*1000.0f, Entry->SizeGB );
+
+	PDBCacheEntries.Remove( Directory );
 }
 
 FPDBCacheEntryRef FPDBCache::FindAndTouchPDBCacheEntry( const FString& PathOrLabel )
 {
-	FPDBCacheEntryRef CacheEntry = PDBCacheEntries.FindChecked( EscapePath( PathOrLabel ) );
+	FPDBCacheEntryRef& CacheEntry = PDBCacheEntries.FindChecked( EscapePath( PathOrLabel ) );
 	TouchPDBCacheEntry( CacheEntry->Directory );
 	return CacheEntry;
 }
