CrashReport - Fixed a typo

diff --git a/Engine/Source/Developer/CrashDebugHelper/Private/CrashDebugHelper.cpp b/Engine/Source/Developer/CrashDebugHelper/Private/CrashDebugHelper.cpp
--- a/Engine/Source/Developer/CrashDebugHelper/Private/CrashDebugHelper.cpp
+++ b/Engine/Source/Developer/CrashDebugHelper/Private/CrashDebugHelper.cpp
@@ -165,7 +165,7 @@ bool ICrashDebugHelper::SyncModules()
 
 		if( bContainsProductVersion )
 		{
-			UE_LOG( LogCrashDebugHelper, Warning, TEXT( "Label %s found in the PDB Cache, using it" ), *CrashInfo.ProductVersion );
+			UE_LOG( LogCrashDebugHelper, Warning, TEXT( "ProductVersion '%s' found in the PDB Cache, using it" ), *CrashInfo.ProductVersion );
 			CrashInfo.PDBCacheEntry = PDBCache.FindAndTouchPDBCacheEntry( CrashInfo.ProductVersion );
 			return true;
 		}
@@ -1176,7 +1176,7 @@ void FPDBCache::RemovePDBCacheEntry( const FString& Directory )
 
 FPDBCacheEntryRef FPDBCache::FindAndTouchPDBCacheEntry( const FString& PathOrLabel )
 {
-	FPDBCacheEntryRef& CacheEntry = PDBCacheEntries.FindChecked( EscapePath( PathOrLabel ) );
+	FPDBCacheEntryRef CacheEntry = PDBCacheEntries.FindChecked( EscapePath( PathOrLabel ) );
 	TouchPDBCacheEntry( CacheEntry->Directory );
 	return CacheEntry;
 }
