Need to switch to symbolicating the crash log in the crash reporter since CoreSymbolication is unsafe within a crash handler. Store in the Apple format so that the work to parse the dump can be re-used to parse crash reports from the Mac App Store where our crash reporter might not work.

diff --git a/Engine/Source/Developer/CrashDebugHelper/Private/Mac/CrashDebugHelperMac.cpp b/Engine/Source/Developer/CrashDebugHelper/Private/Mac/CrashDebugHelperMac.cpp
--- a/Engine/Source/Developer/CrashDebugHelper/Private/Mac/CrashDebugHelperMac.cpp
+++ b/Engine/Source/Developer/CrashDebugHelper/Private/Mac/CrashDebugHelperMac.cpp
@@ -2,10 +2,364 @@
 
 #include "CrashDebugHelperPrivatePCH.h"
 #include "Database.h"
-
+#include "ApplePlatformSymbolication.h"
+#include <cxxabi.h>
 
 #include "ISourceControlModule.h"
 
+static int32 ParseReportVersion(TCHAR const* CrashLog, int32& OutVersion)
+{
+	int32 Found = 0;
+	TCHAR const* VersionLine = FCStringWide::Strstr(CrashLog, TEXT("Report Version:"));
+	if (VersionLine)
+	{
+		Found = swscanf(VersionLine, TEXT("%*ls %*ls %d"), &OutVersion);
+	}
+	return Found;
+}
+
+static int32 ParseVersion(TCHAR const* CrashLog, int32& OutMajor, int32& OutMinor, int32& OutBuild, int32& OutChangeList, FString& OutBranch)
+{
+	int32 Found = 0;
+	TCHAR const* VersionLine = FCStringWide::Strstr(CrashLog, TEXT("Version:"));
+	if (VersionLine)
+	{
+		TCHAR Branch[257] = {0};
+		Found = swscanf(VersionLine, TEXT("%*s %d.%d.%d (%*d.%*d.%*d-%d+%256ls)"), &OutMajor, &OutMinor, &OutBuild, &OutChangeList, Branch);
+		if(Found == 5)
+		{
+			TCHAR* BranchEnd = FCStringWide::Strchr(Branch, TEXT(')'));
+			if(BranchEnd)
+			{
+				*BranchEnd = TEXT('\0');
+			}
+			OutBranch = Branch;
+		}
+	}
+	return Found;
+}
+
+static int32 ParseOS(TCHAR const* CrashLog, int32& OutMajor, int32& OutMinor, int32& OutPatch, int32& OutBuild)
+{
+	int32 Found = 0;
+	TCHAR const* VersionLine = FCStringWide::Strstr(CrashLog, TEXT("OS Version:"));
+	if (VersionLine)
+	{
+		Found = swscanf(VersionLine, TEXT("%*s %*s Mac OS X %d.%d.%d (%xd)"), &OutMajor, &OutMinor, &OutPatch, &OutBuild);
+	}
+	return Found;
+}
+
+static bool ParseModel(TCHAR const* CrashLog, FString& OutModelDetails)
+{
+	bool bFound = false;
+	TCHAR const* Line = FCStringWide::Strstr(CrashLog, TEXT("Model:"));
+	if (Line)
+	{
+		Line += FCStringWide::Strlen(TEXT("Model: "));
+		TCHAR const* End = FCStringWide::Strchr(Line, TEXT('\r'));
+		if(!End)
+		{
+			End = FCStringWide::Strchr(Line, TEXT('\n'));
+		}
+		check(End);
+		
+		int32 Length = FMath::Min(256, (int32)((uintptr_t)(End - Line)));
+		OutModelDetails.Append(Line, Length);
+		
+		bFound = true;
+	}
+	return bFound;
+}
+
+static int32 ParseGraphics(TCHAR const* CrashLog, FString& OutGPUDetails)
+{
+	bool bFound = false;
+	TCHAR const* Line = FCStringWide::Strstr(CrashLog, TEXT("Graphics:"));
+	int32 Output = 0;
+	while (Line)
+	{
+		Line += FCStringWide::Strlen(TEXT("Graphics:"));
+		TCHAR const* End = FCStringWide::Strchr(Line, TEXT('\r'));
+		if(!End)
+		{
+			End = FCStringWide::Strchr(Line, TEXT('\n'));
+		}
+		check(End);
+		
+		int32 Length = FMath::Min((256 - Output), (int32)((uintptr_t)(End - Line)));
+		OutGPUDetails.Append(Line, Length);
+		
+		Line = FCStringWide::Strstr(Line, TEXT("Graphics:"));
+		
+		bFound = true;
+	}
+	return bFound;
+}
+
+static int32 ParseError(TCHAR const* CrashLog, FString& OutErrorDetails)
+{
+	bool bFound = false;
+	TCHAR const* Line = FCStringWide::Strstr(CrashLog, TEXT("Application Specific Information:"));
+	if (Line)
+	{
+		Line = FCStringWide::Strchr(Line, TEXT('\n'));
+		check(Line);
+		Line += 1;
+		TCHAR const* End = FCStringWide::Strchr(Line, TEXT('\r'));
+		if(!End)
+		{
+			End = FCStringWide::Strchr(Line, TEXT('\n'));
+		}
+		check(End);
+		
+		int32 Length = FMath::Min(PATH_MAX, (int32)((uintptr_t)(End - Line)));
+		OutErrorDetails.Append(Line, Length);
+		
+		bFound = true;
+	}
+	return bFound;
+}
+
+static int32 ParseCrashedThread(TCHAR const* CrashLog, int32& OutThreadNumber)
+{
+	int32 Found = 0;
+	TCHAR const* Line = FCStringWide::Strstr(CrashLog, TEXT("Crashed Thread:"));
+	if (Line)
+	{
+		Found = swscanf(Line, TEXT("%*s %*s %d"), &OutThreadNumber);
+	}
+	return Found;
+}
+
+static TCHAR const* FindThreadStack(TCHAR const* CrashLog, int32 const ThreadNumber)
+{
+	int32 Found = 0;
+	FString Format = FString::Printf(TEXT("Thread %d"), ThreadNumber);
+	TCHAR const* Line = FCStringWide::Strstr(CrashLog, *Format);
+	if (Line)
+	{
+		Line = FCStringWide::Strchr(Line, TEXT('\n'));
+		check(Line);
+		Line += 1;
+	}
+	return Line;
+}
+
+static TCHAR const* FindCrashedThreadStack(TCHAR const* CrashLog)
+{
+	TCHAR const* Line = nullptr;
+	int32 ThreadNumber = 0;
+	int32 Found = ParseCrashedThread(CrashLog, ThreadNumber);
+	if(Found)
+	{
+		Line = FindThreadStack(CrashLog, ThreadNumber);
+	}
+	return Line;
+}
+
+static int32 ParseThreadStackLine(TCHAR const* StackLine, FString& OutModuleName, uint64& OutProgramCounter, FString& OutFunctionName, FString& OutFileName, int32& OutLineNumber)
+{
+	TCHAR ModuleName[257];
+	TCHAR FunctionName[1025];
+	TCHAR FileName[257];
+	
+	int32 Found = swscanf(StackLine, TEXT("%*d %256ls 0x%lx %1024ls + %*d (%256ls:%d)"), ModuleName, &OutProgramCounter, FunctionName, FileName, &OutLineNumber);
+	switch(Found)
+	{
+		case 5:
+		case 4:
+		{
+			OutFileName = FileName;
+		}
+		case 3:
+		{
+			int32 Status = -1;
+			ANSICHAR* DemangledName = abi::__cxa_demangle(TCHAR_TO_UTF8(FunctionName), nullptr, nullptr, &Status);
+			if(DemangledName && Status == 0)
+			{
+				// C++ function
+				OutFunctionName = FString::Printf(TEXT("%ls "), UTF8_TO_TCHAR(DemangledName));
+			}
+			else if (FCStringWide::Strlen(FunctionName) > 0 && FCStringWide::Strchr(FunctionName, ']'))
+			{
+				// ObjC function
+				OutFunctionName = FString::Printf(TEXT("%ls "), FunctionName);
+			}
+			else if(FCStringWide::Strlen(FunctionName) > 0)
+			{
+				// C function
+				OutFunctionName = FString::Printf(TEXT("%ls() "), FunctionName);
+			}
+		}
+		case 2:
+		case 1:
+		{
+			OutModuleName = ModuleName;
+		}
+		default:
+		{
+			break;
+		}
+	}
+	
+	return Found;
+}
+
+static int32 SymboliseStackInfo(TArray<FCrashModuleInfo> const& ModuleInfo, FString ModuleName, uint64 const ProgramCounter, FString& OutFunctionName, FString& OutFileName, int32& OutLineNumber)
+{
+	FProgramCounterSymbolInfo Info;
+	int32 ValuesSymbolised = 0;
+	
+	FCrashModuleInfo Module;
+	for (auto Iterator : ModuleInfo)
+	{
+		if(Iterator.Name.EndsWith(ModuleName))
+		{
+			Module = Iterator;
+			break;
+		}
+	}
+	check(Module.Name.Len() > 0);
+	
+	bool bSuccess = FApplePlatformSymbolication::SymbolInfoForStrippedSymbol((ProgramCounter - Module.BaseOfImage), TCHAR_TO_UTF8(*Module.Name), TCHAR_TO_UTF8(*Module.Report), Info);
+	
+	if(bSuccess)
+	{
+		if(FCStringAnsi::Strlen(Info.FunctionName) > 0)
+		{
+			OutFunctionName = Info.FunctionName;
+			ValuesSymbolised++;
+		}
+		if(ValuesSymbolised == 1 && FCStringAnsi::Strlen(Info.Filename) > 0)
+		{
+			OutFileName = Info.Filename;
+			ValuesSymbolised++;
+		}
+		if(ValuesSymbolised == 2 && Info.LineNumber > 0)
+		{
+			OutLineNumber = Info.LineNumber;
+			ValuesSymbolised++;
+		}
+	}
+	
+	return ValuesSymbolised;
+}
+
+static TCHAR const* FindModules(TCHAR const* CrashLog)
+{
+	TCHAR const* Line = FCStringWide::Strstr(CrashLog, TEXT("Binary Images:"));
+	if (Line)
+	{
+		Line = FCStringWide::Strchr(Line, TEXT('\n'));
+		check(Line);
+		Line += 1;
+	}
+	return Line;
+}
+
+static int32 ParseModuleVersion(TCHAR const* Version, uint16& OutMajor, uint16& OutMinor, uint16& OutPatch, uint16& OutBuild)
+{
+	OutMajor = OutMinor = OutPatch = OutBuild = 0;
+	int32 Found = swscanf(Version, TEXT("%hu.%hu.%hu"), &OutMajor, &OutMinor, &OutPatch);
+	TCHAR const* CurrentStart = FCStringWide::Strchr(Version, TEXT('-'));
+	if(CurrentStart)
+	{
+		int32 Components[3] = {0, 0, 0};
+		int32 Result = swscanf(CurrentStart, TEXT("%*ls %d.%d.%d"), &Components[0], &Components[1], &Components[2]);
+		OutBuild = (uint16)(Components[0] * 10000) + (Components[1] * 100) + (Components[2]);
+		Found = 4;
+	}
+	return Found;
+}
+
+static bool ParseModuleLine(TCHAR const* ModuleLine, FCrashModuleInfo& OutModule)
+{
+	bool bOK = false;
+	TCHAR ModuleName[257] = {0};
+	uint64 ModuleBase = 0;
+	uint64 ModuleEnd = 0;
+	
+	int32 Found = swscanf(ModuleLine, TEXT("%lx %*ls %lx %256ls"), &ModuleBase, &ModuleEnd, ModuleName);
+	switch (Found)
+	{
+		case 3:
+		{
+			TCHAR const* VersionStart = FCStringWide::Strchr(ModuleLine, TEXT('('));
+			TCHAR const* VersionEnd = FCStringWide::Strchr(ModuleLine, TEXT(')'));
+			if(VersionStart && VersionEnd)
+			{
+				++VersionStart;
+				Found += ParseModuleVersion(VersionStart, OutModule.Major, OutModule.Minor, OutModule.Patch, OutModule.Revision);
+			}
+			
+			TCHAR const* UUIDStart = FCStringWide::Strchr(ModuleLine, TEXT('<'));
+			TCHAR const* UUIDEnd = FCStringWide::Strchr(ModuleLine, TEXT('>'));
+			if(UUIDStart && UUIDEnd)
+			{
+				++UUIDStart;
+				int32 Length = FMath::Min(64, (int32)((uintptr_t)(UUIDEnd - UUIDStart)));
+				OutModule.Report.Append(UUIDStart, Length);
+				Found++;
+			}
+			
+			TCHAR const* Path = FCStringWide::Strchr(ModuleLine, TEXT('/'));
+			if(Path)
+			{
+				TCHAR const* End = FCStringWide::Strchr(Path, TEXT('\r'));
+				if(!End)
+				{
+					End = FCStringWide::Strchr(Path, TEXT('\n'));
+				}
+				check(End);
+				
+				int32 Length = FMath::Min(PATH_MAX, (int32)((uintptr_t)(End - Path)));
+				OutModule.Name.Append(Path, Length);
+				Found++;
+				bOK = true;
+			}
+		}
+		case 2:
+		{
+			OutModule.SizeOfImage = (ModuleBase - ModuleEnd);
+		}
+		case 1:
+		{
+			OutModule.BaseOfImage = ModuleBase;
+			break;
+		}
+		default:
+		{
+			break;
+		}
+	}
+	return bOK;
+}
+
+static int64 StringSize( const ANSICHAR* Line )
+{
+	int64 Size = 0;
+	if( Line != NULL )
+	{
+		while( *Line++ != 0 )
+		{
+			Size++;
+		}
+	}
+	return Size;
+}
+
+static void WriteLine( FArchive* ReportFile, const ANSICHAR* Line = "" )
+{
+	if( Line != NULL )
+	{
+		int64 StringBytes = StringSize( Line );
+		ReportFile->Serialize( ( void* )Line, StringBytes );
+	}
+	
+	ReportFile->Serialize( TCHAR_TO_UTF8( LINE_TERMINATOR ), 1 );
+}
+
 FCrashDebugHelperMac::FCrashDebugHelperMac()
 {
 }
@@ -21,10 +375,40 @@ bool FCrashDebugHelperMac::ParseCrashDump(const FString& InCrashDumpName, FCrash
 		UE_LOG(LogCrashDebugHelper, Warning, TEXT("ParseCrashDump: CrashDebugHelper not initialized"));
 		return false;
 	}
-
-	//if (ParseCrashDump_Mac(*InCrashDumpName, OutCrashDebugInfo) == true)
+	
+	FString CrashDump;
+	if ( FFileHelper::LoadFileToString( CrashDump, *InCrashDumpName ) )
 	{
-		return true;
+		// Only supports Apple crash report version 11
+		int32 ReportVersion = 0;
+		int32 Result = ParseReportVersion(*CrashDump, ReportVersion);
+		if(Result == 1 && ReportVersion == 11)
+		{
+			int32 Major = 0;
+			int32 Minor = 0;
+			int32 Build = 0;
+			int32 CLNumber = 0;
+			FString Branch;
+			int32 Result = ParseVersion(*CrashDump, Major, Minor, Build, CLNumber, Branch);
+			if(Result >= 3)
+			{
+				if (Result < 5)
+				{
+					OutCrashDebugInfo.EngineVersion = Build;
+				}
+				else
+				{
+					OutCrashDebugInfo.EngineVersion = CLNumber;
+				}
+				if(Result == 5)
+				{
+					OutCrashDebugInfo.SourceControlLabel = Branch;
+				}
+				OutCrashDebugInfo.PlatformName = TEXT("Mac");
+				OutCrashDebugInfo.CrashDumpName = InCrashDumpName;
+				return true;
+			}
+		}
 	}
 
 	return false;
@@ -52,5 +436,191 @@ bool FCrashDebugHelperMac::SyncAndDebugCrashDump(const FString& InCrashDumpName)
 
 bool FCrashDebugHelperMac::CreateMinidumpDiagnosticReport( const FString& InCrashDumpName )
 {
-	return true;
+	bool bOK = false;
+	
+	FString CrashDump;
+	if ( FFileHelper::LoadFileToString( CrashDump, *InCrashDumpName ) )
+	{
+		int32 ReportVersion = 0;
+		int32 Result = ParseReportVersion(*CrashDump, ReportVersion);
+		if(Result == 1 && ReportVersion == 11)
+		{
+			FString DiagnosticsPath = FPaths::GetPath(InCrashDumpName) / TEXT("diagnostics.txt");
+			
+			FArchive* ReportFile = IFileManager::Get().CreateFileWriter( *DiagnosticsPath );
+			if( ReportFile != NULL )
+			{
+				FString Error;
+				FString ModulePath;
+				FString ModuleName;
+				FString FunctionName;
+				FString FileName;
+				FString Branch;
+				FString Model;
+				FString Gpu;
+				FString Line;
+				
+				uint64 ProgramCounter = 0;
+				int32 Major = 0;
+				int32 Minor = 0;
+				int32 Build = 0;
+				int32 CLNumber = 0;
+				int32 OSMajor = 0;
+				int32 OSMinor = 0;
+				int32 OSPatch = 0;
+				int32 OSBuild = 0;
+				int32 LineNumber = 0;;
+				
+				WriteLine( ReportFile, TCHAR_TO_UTF8( TEXT( "Generating report for minidump" ) ) );
+				WriteLine( ReportFile );
+				
+				int32 Result = ParseVersion(*CrashDump, Major, Minor, Build, CLNumber, Branch);
+				if(Result >= 3)
+				{
+					Line = FString::Printf( TEXT( "Application version %d.%d.%d" ), Major, Minor, Build );
+					WriteLine( ReportFile, TCHAR_TO_UTF8( *Line ) );
+				}
+				
+				if(Result >= 4)
+				{
+					Line = FString::Printf( TEXT( " ... built from changelist %d" ), CLNumber );
+					WriteLine( ReportFile, TCHAR_TO_UTF8( *Line ) );
+				}
+				
+				if(Result == 5 && Branch.Len() > 0)
+				{
+					Line = FString::Printf( TEXT( " ... based on label %s" ), *Branch );
+					WriteLine( ReportFile, TCHAR_TO_UTF8( *Line ) );
+				}
+				WriteLine( ReportFile );
+				
+				Result = ParseOS(*CrashDump, OSMajor, OSMinor, OSPatch, OSBuild);
+				check(Result == 4);
+				Line = FString::Printf( TEXT( "OS version %d.%d.%d.%xd" ), OSMajor, OSMinor, OSPatch, OSBuild );
+				WriteLine( ReportFile, TCHAR_TO_UTF8( *Line ) );
+				
+				ParseModel(*CrashDump, Model);
+				ParseGraphics(*CrashDump, Gpu);
+				
+				Line = FString::Printf( TEXT( "Running %s %s" ), *Model, *Gpu );
+				WriteLine( ReportFile, TCHAR_TO_UTF8( *Line ) );
+				
+				Result = ParseError(*CrashDump, Error);
+				check(Result == 1);
+				
+				Line = FString::Printf( TEXT( "Exception was \"%s\"" ), *Error );
+				WriteLine( ReportFile, TCHAR_TO_UTF8( *Line ) );
+				WriteLine( ReportFile );
+				
+				// Parse modules now for symbolication - if we don't have the running process we need to symbolicate by UUID
+				TArray<FCrashModuleInfo> Modules;
+				TCHAR const* ModuleLine = FindModules(*CrashDump);
+				while(ModuleLine)
+				{
+					FCrashModuleInfo Module;
+					bool const bOK = ParseModuleLine(ModuleLine, Module);
+					if(bOK)
+					{
+						Modules.Push(Module);
+						
+						ModuleLine = FCStringWide::Strchr(ModuleLine, TEXT('\n'));
+						check(ModuleLine);
+						ModuleLine += 1;
+					}
+					else
+					{
+						ModuleLine = nullptr;
+					}
+				}
+				
+				Line = FString::Printf( TEXT( "<CALLSTACK START>" ) );
+				WriteLine( ReportFile, TCHAR_TO_UTF8( *Line ) );
+				
+				TCHAR const* ThreadStackLine = FindCrashedThreadStack(*CrashDump);
+				while(ThreadStackLine)
+				{
+					Result = ParseThreadStackLine(ThreadStackLine, ModuleName, ProgramCounter, FunctionName, FileName, LineNumber);
+					
+					// If we got the modulename & program counter but didn't parse the filename & linenumber we can resymbolise
+					if(Result > 1 && Result < 4)
+					{
+						// Attempt to resymbolise using CoreSymbolication
+						Result += SymboliseStackInfo(Modules, ModuleName, ProgramCounter, FunctionName, FileName, LineNumber);
+					}
+					
+					// Output in our format based on the fields we actually have
+					switch (Result)
+					{
+						case 2:
+							Line = FString::Printf( TEXT( "Unknown() Address = 0x%lx (filename not found) [in %s]" ), ProgramCounter, *ModuleName );
+							WriteLine( ReportFile, TCHAR_TO_UTF8( *Line ) );
+							
+							ThreadStackLine = FCStringWide::Strchr(ThreadStackLine, TEXT('\n'));
+							check(ThreadStackLine);
+							ThreadStackLine += 1;
+							break;
+							
+						case 3:
+						case 4:
+							Line = FString::Printf( TEXT( "%s Address = 0x%lx (filename not found) [in %s]" ), *FunctionName, ProgramCounter, *ModuleName );
+							WriteLine( ReportFile, TCHAR_TO_UTF8( *Line ) );
+							
+							ThreadStackLine = FCStringWide::Strchr(ThreadStackLine, TEXT('\n'));
+							check(ThreadStackLine);
+							ThreadStackLine += 1;
+							break;
+							
+						case 5:
+						case 6: // Function name might be parsed twice
+							Line = FString::Printf( TEXT( "%s Address = 0x%lx [%s, line %d] [in %s]" ), *FunctionName, ProgramCounter, *FileName, LineNumber, *ModuleName );
+							WriteLine( ReportFile, TCHAR_TO_UTF8( *Line ) );
+							
+							ThreadStackLine = FCStringWide::Strchr(ThreadStackLine, TEXT('\n'));
+							check(ThreadStackLine);
+							ThreadStackLine += 1;
+							break;
+							
+						default:
+							ThreadStackLine = nullptr;
+							break;
+					}
+				}
+				
+				Line = FString::Printf( TEXT( "<CALLSTACK END>" ) );
+				WriteLine( ReportFile, TCHAR_TO_UTF8( *Line ) );
+				WriteLine( ReportFile );
+				
+				Line = FString::Printf( TEXT( "%d loaded modules" ), Modules.Num() );
+				WriteLine( ReportFile, TCHAR_TO_UTF8( *Line ) );
+				
+				for( int32 ModuleIndex = 0; ModuleIndex < Modules.Num(); ModuleIndex++ )
+				{
+					FCrashModuleInfo& Module = Modules[ModuleIndex];
+					
+					FString ModuleDirectory = FPaths::GetPath(Module.Name);
+					FString ModuleName = FPaths::GetBaseFilename( Module.Name, true ) + FPaths::GetExtension( Module.Name, true );
+					
+					FString ModuleDetail = FString::Printf( TEXT( "%40s" ), *ModuleName );
+					FString Version = FString::Printf( TEXT( " (%d.%d.%d.%d)" ), Module.Major, Module.Minor, Module.Patch, Module.Revision );
+					ModuleDetail += FString::Printf( TEXT( " %22s" ), *Version );
+					ModuleDetail += FString::Printf( TEXT( " 0x%016x 0x%08x" ), Module.BaseOfImage, Module.SizeOfImage );
+					ModuleDetail += FString::Printf( TEXT( " %s" ), *ModuleDirectory );
+					
+					WriteLine( ReportFile, TCHAR_TO_UTF8( *ModuleDetail ) );
+				}
+				
+				WriteLine( ReportFile );
+				
+				Line = FString::Printf( TEXT( "Report end!" ) );
+				WriteLine( ReportFile, TCHAR_TO_UTF8( *Line )  );
+				
+				ReportFile->Close();
+				delete ReportFile;
+				
+				bOK = true;
+			}
+		}
+	}
+	
+	return bOK;
 }
diff --git a/Engine/Source/Programs/CrashReporter/CrashReportClient/Private/Mac/MacErrorReport.cpp b/Engine/Source/Programs/CrashReporter/CrashReportClient/Private/Mac/MacErrorReport.cpp
--- a/Engine/Source/Programs/CrashReporter/CrashReportClient/Private/Mac/MacErrorReport.cpp
+++ b/Engine/Source/Programs/CrashReporter/CrashReportClient/Private/Mac/MacErrorReport.cpp
@@ -2,14 +2,31 @@
 
 #include "CrashReportClientApp.h"
 #include "../CrashReportUtil.h"
+#include "CrashDebugHelperModule.h"
 
 #include "MacErrorReport.h"
 
+namespace
+{
+	/** Pointer to dynamically loaded crash diagnosis module */
+	FCrashDebugHelperModule* CrashHelperModule;
+}
+
 FMacErrorReport::FMacErrorReport(const FString& Directory)
 	: FGenericErrorReport(Directory)
 {
 }
 
+void FMacErrorReport::Init()
+{
+	CrashHelperModule = &FModuleManager::LoadModuleChecked<FCrashDebugHelperModule>(FName("CrashDebugHelper"));
+}
+
+void FMacErrorReport::ShutDown()
+{
+	CrashHelperModule->ShutdownModule();
+}
+
 FString FMacErrorReport::FindCrashedAppName() const
 {
 	TArray<uint8> Data;
@@ -71,3 +88,53 @@ FString FMacErrorReport::FindMostRecentErrorReport()
 
 	return RecentReportDirectory;
 }
+
+FText FMacErrorReport::DiagnoseReport() const
+{
+	// Should check if there are local PDBs before doing anything
+	auto CrashDebugHelper = CrashHelperModule ? CrashHelperModule->Get() : nullptr;
+	if (!CrashDebugHelper)
+	{
+		// Not localized: should never be seen
+		return FText::FromString(TEXT("Failed to load CrashDebugHelper."));
+	}
+	
+	FString DumpFilename;
+	if (!FindFirstReportFileWithExtension(DumpFilename, TEXT(".dmp")))
+	{
+		if (!FindFirstReportFileWithExtension(DumpFilename, TEXT(".mdmp")))
+		{
+			return FText::FromString("No minidump found for this crash.");
+		}
+	}
+	
+	FCrashDebugInfo DebugInfo;
+	if (!CrashDebugHelper->ParseCrashDump(ReportDirectory / DumpFilename, DebugInfo))
+	{
+		return FText::FromString("No minidump found for this crash.");
+	}
+	
+	if (!CrashDebugHelper->CreateMinidumpDiagnosticReport(ReportDirectory / DumpFilename))
+	{
+		if ( FRocketSupport::IsRocket() )
+		{
+			return FText::FromString("We apologize for the inconvenience.\nPlease send this crash report to help improve our software.");
+		}
+		else
+		{
+			return FText::FromString("You do not have any debugging symbols required to display the callstack for this crash.");
+		}
+	}
+	else
+	{
+		FString CrashDump;
+		if ( FFileHelper::LoadFileToString( CrashDump, *(ReportDirectory / TEXT("diagnostics.txt")) ) )
+		{
+			return FText::FromString(CrashDump);
+		}
+		else
+		{
+			return FText::FromString("Failed to create diagnosis information.");
+		}
+	}
+}
diff --git a/Engine/Source/Programs/CrashReporter/CrashReportClient/Private/Mac/MacErrorReport.h b/Engine/Source/Programs/CrashReporter/CrashReportClient/Private/Mac/MacErrorReport.h
--- a/Engine/Source/Programs/CrashReporter/CrashReportClient/Private/Mac/MacErrorReport.h
+++ b/Engine/Source/Programs/CrashReporter/CrashReportClient/Private/Mac/MacErrorReport.h
@@ -16,6 +16,16 @@ public:
 	FMacErrorReport()
 	{
 	}
+	
+	/**
+	 * Load helper modules
+	 */
+	static void Init();
+	
+	/**
+	 * Unload helper modules
+	 */
+	static void ShutDown();
 
 	/**
 	 * Discover all files in the crash report directory
@@ -24,13 +34,10 @@ public:
 	explicit FMacErrorReport(const FString& Directory);
 
 	/**
-	 * Do nothing - shouldn't be called on Mac
-	 * @return Dummy text
+	 * Parse the callstack from the Apple-style crash report log
+	 * @return UE4 crash diagnosis text
 	 */
-	FText DiagnoseReport() const	
-	{
-		return FText::FromString("No local diagnosis on Mac");
-	}
+	FText DiagnoseReport() const;
 
 	/**
 	 * Get the name of the crashed app from the report (hides implementation in FGenericErrorReport)
diff --git a/Engine/Source/Runtime/Core/Private/Apple/ApplePlatformCrashContext.cpp b/Engine/Source/Runtime/Core/Private/Apple/ApplePlatformCrashContext.cpp
--- a/Engine/Source/Runtime/Core/Private/Apple/ApplePlatformCrashContext.cpp
+++ b/Engine/Source/Runtime/Core/Private/Apple/ApplePlatformCrashContext.cpp
@@ -161,15 +161,16 @@ void FApplePlatformCrashContext::WriteLine(int ReportFile, const TCHAR* Line)
 	WriteUTF16String(ReportFile, WindowsTerminator, 2);
 }
 
-ANSICHAR* FApplePlatformCrashContext::ItoANSI(uint64 Val, uint64 Base)
+ANSICHAR* FApplePlatformCrashContext::ItoANSI(uint64 Val, uint64 Base, uint32 Len)
 {
 	static ANSICHAR InternalBuffer[64] = {0};
 	
 	uint64 i = 62;
+	int32 pad = Len;
 	
 	if(Val)
 	{
-		for(; Val && i ; --i, Val /= Base)
+		for(; Val && i ; --i, Val /= Base, --pad)
 		{
 			InternalBuffer[i] = "0123456789abcdef"[Val % Base];
 		}
@@ -177,20 +178,28 @@ ANSICHAR* FApplePlatformCrashContext::ItoANSI(uint64 Val, uint64 Base)
 	else
 	{
 		InternalBuffer[i--] = '0';
+		--pad;
+	}
+	
+	while(pad > 0)
+	{
+		InternalBuffer[i--] = '0';
+		--pad;
 	}
 	
 	return &InternalBuffer[i+1];
 }
 
-TCHAR* FApplePlatformCrashContext::ItoTCHAR(uint64 Val, uint64 Base)
+TCHAR* FApplePlatformCrashContext::ItoTCHAR(uint64 Val, uint64 Base, uint32 Len)
 {
 	static TCHAR InternalBuffer[64] = {0};
 	
 	uint64 i = 62;
+	int32 pad = Len;
 	
 	if(Val)
 	{
-		for(; Val && i ; --i, Val /= Base)
+		for(; Val && i ; --i, Val /= Base, --pad)
 		{
 			InternalBuffer[i] = TEXT("0123456789abcdef")[Val % Base];
 		}
@@ -198,6 +207,13 @@ TCHAR* FApplePlatformCrashContext::ItoTCHAR(uint64 Val, uint64 Base)
 	else
 	{
 		InternalBuffer[i--] = TEXT('0');
+		--pad;
+	}
+	
+	while(pad > 0)
+	{
+		InternalBuffer[i--] = '0';
+		--pad;
 	}
 	
 	return &InternalBuffer[i+1];
diff --git a/Engine/Source/Runtime/Core/Private/Apple/ApplePlatformStackWalk.cpp b/Engine/Source/Runtime/Core/Private/Apple/ApplePlatformStackWalk.cpp
--- a/Engine/Source/Runtime/Core/Private/Apple/ApplePlatformStackWalk.cpp
+++ b/Engine/Source/Runtime/Core/Private/Apple/ApplePlatformStackWalk.cpp
@@ -14,12 +14,21 @@
 // Internal helper functions not exposed to public
 int32 GetModuleImageSize( const struct mach_header* Header )
 {
+#if PLATFORM_64BITS
+	struct load_command *CurrentCommand = (struct load_command *)( (char *)Header + sizeof(struct mach_header_64) );
+	int32 ModuleSize = 0;
+	
+	// Check header
+	if( Header->magic != MH_MAGIC_64 )
+		return 0;
+#else
 	struct load_command *CurrentCommand = (struct load_command *)( (char *)Header + sizeof(struct mach_header) );
 	int32 ModuleSize = 0;
 	
 	// Check header
 	if( Header->magic != MH_MAGIC )
 		return 0;
+#endif
 	
 	for( int32 i = 0; i < Header->ncmds; i++ )
 	{
@@ -37,11 +46,21 @@ int32 GetModuleImageSize( const struct mach_header* Header )
 
 int32 GetModuleTimeStamp( const struct mach_header* Header )
 {
+#if PLATFORM_64BITS
+	struct load_command *CurrentCommand = (struct load_command *)( (char *)Header + sizeof(struct mach_header_64) );
+	int32 ModuleSize = 0;
+	
+	// Check header
+	if( Header->magic != MH_MAGIC_64 )
+		return 0;
+#else
 	struct load_command *CurrentCommand = (struct load_command *)( (char *)Header + sizeof(struct mach_header) );
+	int32 ModuleSize = 0;
 	
 	// Check header
 	if( Header->magic != MH_MAGIC )
 		return 0;
+#endif
 	
 	for( int32 i = 0; i < Header->ncmds; i++ )
 	{
@@ -57,6 +76,75 @@ int32 GetModuleTimeStamp( const struct mach_header* Header )
 	return 0;
 }
 
+static void AsyncSafeProgramCounterToSymbolInfo( uint64 ProgramCounter, FProgramCounterSymbolInfo&  SymbolInfo )
+{
+	Dl_info DylibInfo;
+	int32 Result = dladdr((const void*)ProgramCounter, &DylibInfo);
+	if (Result == 0)
+	{
+		return;
+	}
+
+#if PLATFORM_MAC // On the Mac the crash report client can resymbolise
+	if (DylibInfo.dli_sname)
+	{
+		// Mangled function name
+		FCStringAnsi::Sprintf(SymbolInfo.FunctionName, "%s ", DylibInfo.dli_sname);
+	}
+	else
+	{
+		// Unknown!
+		FCStringAnsi::Sprintf(SymbolInfo.FunctionName, "[Unknown]() ");
+	}
+#else // But on iOS the best we can do is demangle
+	int32 Status = 0;
+	ANSICHAR* DemangledName = NULL;
+	
+	// Increased the size of the demangle destination to reduce the chances that abi::__cxa_demangle will allocate
+	// this causes the app to hang as malloc isn't signal handler safe. Ideally we wouldn't call this function in a handler.
+	size_t DemangledNameLen = 65536;
+	ANSICHAR DemangledNameBuffer[65536]= {0};
+	DemangledName = abi::__cxa_demangle(DylibInfo.dli_sname, DemangledNameBuffer, &DemangledNameLen, &Status);
+	
+	if (DemangledName)
+	{
+		// C++ function
+		FCStringAnsi::Sprintf(SymbolInfo.FunctionName, "%s ", DemangledName);
+	}
+	else if (DylibInfo.dli_sname && strchr(DylibInfo.dli_sname, ']'))
+	{
+		// ObjC function
+		FCStringAnsi::Sprintf(SymbolInfo.FunctionName, "%s ", DylibInfo.dli_sname);
+	}
+	else if(DylibInfo.dli_sname)
+	{
+		// C function
+		FCStringAnsi::Sprintf(SymbolInfo.FunctionName, "%s() ", DylibInfo.dli_sname);
+	}
+	else
+	{
+		// Unknown!
+		FCStringAnsi::Sprintf(SymbolInfo.FunctionName, "[Unknown]() ");
+	}
+#endif
+	
+	// No line number found.
+	FCStringAnsi::Strcat(SymbolInfo.Filename, "Unknown");
+	SymbolInfo.LineNumber = 0;
+	
+	// Write out Module information.
+	ANSICHAR* DylibPath = (ANSICHAR*)DylibInfo.dli_fname;
+	ANSICHAR* DylibName = FCStringAnsi::Strrchr(DylibPath, '/');
+	if (DylibName)
+	{
+		DylibName += 1;
+	}
+	else
+	{
+		DylibName = DylibPath;
+	}
+	FCStringAnsi::Strcpy(SymbolInfo.ModuleName, DylibName);
+}
 
 void FApplePlatformStackWalk::CaptureStackBackTrace( uint64* BackTrace, uint32 MaxDepth, void* Context )
 {
@@ -80,7 +168,14 @@ bool FApplePlatformStackWalk::ProgramCounterToHumanReadableString( int32 Current
 	}
 
 	FProgramCounterSymbolInfo SymbolInfo;
-	ProgramCounterToSymbolInfo( ProgramCounter, SymbolInfo );
+	if(!Context)
+	{
+		ProgramCounterToSymbolInfo( ProgramCounter, SymbolInfo );
+	}
+	else
+	{
+		AsyncSafeProgramCounterToSymbolInfo(ProgramCounter, SymbolInfo);
+	}
 	
 	// Write out function name.
 	FCStringAnsi::Strcat(HumanReadableString, HumanReadableStringSize, SymbolInfo.FunctionName);
@@ -122,59 +217,7 @@ void FApplePlatformStackWalk::ProgramCounterToSymbolInfo( uint64 ProgramCounter,
 	bool bOK = FApplePlatformSymbolication::SymbolInfoForAddress(ProgramCounter, SymbolInfo);
 	if(!bOK)
 	{
-		Dl_info DylibInfo;
-		int32 Result = dladdr((const void*)ProgramCounter, &DylibInfo);
-		if (Result == 0)
-		{
-			return;
-		}
-		
-		int32 Status = 0;
-		ANSICHAR* DemangledName = NULL;
-		
-		// Increased the size of the demangle destination to reduce the chances that abi::__cxa_demangle will allocate
-		// this causes the app to hang as malloc isn't signal handler safe. Ideally we wouldn't call this function in a handler.
-		size_t DemangledNameLen = 65536;
-		ANSICHAR DemangledNameBuffer[65536]= {0};
-		DemangledName = abi::__cxa_demangle(DylibInfo.dli_sname, DemangledNameBuffer, &DemangledNameLen, &Status);
-		
-		if (DemangledName)
-		{
-			// C++ function
-			FCStringAnsi::Sprintf(SymbolInfo.FunctionName, "%s ", DemangledName);
-		}
-		else if (DylibInfo.dli_sname && strchr(DylibInfo.dli_sname, ']'))
-		{
-			// ObjC function
-			FCStringAnsi::Sprintf(SymbolInfo.FunctionName, "%s ", DylibInfo.dli_sname);
-		}
-		else if(DylibInfo.dli_sname)
-		{
-			// C function
-			FCStringAnsi::Sprintf(SymbolInfo.FunctionName, "%s() ", DylibInfo.dli_sname);
-		}
-		else
-		{
-			// Unknown!
-			FCStringAnsi::Sprintf(SymbolInfo.FunctionName, "[Unknown]() ");
-		}
-		
-		// No line number found.
-		FCStringAnsi::Strcat(SymbolInfo.Filename, "Unknown");
-		SymbolInfo.LineNumber = 0;
-		
-		// Write out Module information.
-		ANSICHAR* DylibPath = (ANSICHAR*)DylibInfo.dli_fname;
-		ANSICHAR* DylibName = FCStringAnsi::Strrchr(DylibPath, '/');
-		if(DylibName)
-		{
-			DylibName += 1;
-		}
-		else
-		{
-			DylibName = DylibPath;
-		}
-		FCStringAnsi::Strcpy(SymbolInfo.ModuleName, DylibName);
+		AsyncSafeProgramCounterToSymbolInfo(ProgramCounter, SymbolInfo);
 	}
 }
 
diff --git a/Engine/Source/Runtime/Core/Private/Apple/ApplePlatformSymbolication.cpp b/Engine/Source/Runtime/Core/Private/Apple/ApplePlatformSymbolication.cpp
--- a/Engine/Source/Runtime/Core/Private/Apple/ApplePlatformSymbolication.cpp
+++ b/Engine/Source/Runtime/Core/Private/Apple/ApplePlatformSymbolication.cpp
@@ -39,10 +39,12 @@ extern "C"
 	void CSRelease(CSTypeRef CS);
 	void CSShow(CSTypeRef CS);
 	
-	CSSymbolicatorRef CSSymbolicatorCreateWithTask(task_t Task);
+	CSSymbolicatorRef CSSymbolicatorCreateWithPid(pid_t pid);
+	CSSymbolicatorRef CSSymbolicatorCreateWithPathAndArchitecture(const char* path, cpu_type_t type);
 	
 	CSSourceInfoRef CSSymbolicatorGetSourceInfoWithAddressAtTime(CSSymbolicatorRef Symbolicator, vm_address_t Address, uint64_t Time);
 	CSSymbolRef CSSymbolicatorGetSymbolWithMangledNameFromSymbolOwnerWithNameAtTime(CSSymbolicatorRef Symbolicator, const char* Symbol, const char* Name, uint64_t Time);
+	CSSymbolOwnerRef CSSymbolicatorGetSymbolOwnerWithUUIDAtTime(CSSymbolicatorRef Symbolicator, CFUUIDRef UUID, uint64_t Time);
 	
 	const char* CSSymbolGetName(CSSymbolRef Symbol);
 	CSRange CSSymbolGetRange(CSSymbolRef Symbol);
@@ -61,7 +63,8 @@ bool FApplePlatformSymbolication::SymbolInfoForAddress(uint64 ProgramCounter, FP
 {
 #if PLATFORM_MAC
 	bool bOK = false;
-	CSSymbolicatorRef Symbolicator = CSSymbolicatorCreateWithTask(mach_task_self());
+	
+	CSSymbolicatorRef Symbolicator = CSSymbolicatorCreateWithPid(FPlatformProcess::GetCurrentProcessId());
 	if(!CSIsNull(Symbolicator))
 	{
 		CSSourceInfoRef Symbol = CSSymbolicatorGetSourceInfoWithAddressAtTime(Symbolicator, (vm_address_t)ProgramCounter, kCSNow);
@@ -81,15 +84,12 @@ bool FApplePlatformSymbolication::SymbolInfoForAddress(uint64 ProgramCounter, FP
 				FCStringAnsi::Strcpy(Info.ModuleName, DylibName);
 				
 				bOK = Info.LineNumber != 0;
-				
-				CSRelease(Owner);
 			}
-			
-			CSRelease(Symbol);
 		}
 		
 		CSRelease(Symbolicator);
 	}
+	
 	return bOK;
 #else
 	return false;
@@ -100,7 +100,8 @@ bool FApplePlatformSymbolication::SymbolInfoForFunctionFromModule(ANSICHAR const
 {
 #if PLATFORM_MAC
 	bool bOK = false;
-	CSSymbolicatorRef Symbolicator = CSSymbolicatorCreateWithTask(mach_task_self());
+	
+	CSSymbolicatorRef Symbolicator = CSSymbolicatorCreateWithPid(FPlatformProcess::GetCurrentProcessId());
 	if(!CSIsNull(Symbolicator))
 	{
 		CSSymbolRef Symbol = CSSymbolicatorGetSymbolWithMangledNameFromSymbolOwnerWithNameAtTime(Symbolicator, MangledName, ModuleName, kCSNow);
@@ -110,12 +111,75 @@ bool FApplePlatformSymbolication::SymbolInfoForFunctionFromModule(ANSICHAR const
 			CSRange CodeRange = CSSymbolGetRange(Symbol);
 		
 			bOK = SymbolInfoForAddress(CodeRange.Location, Info);
-	
-			CSRelease(Symbol);
 		}
 		
 		CSRelease(Symbolicator);
 	}
+	
+	return bOK;
+#else
+	return false;
+#endif
+}
+
+bool FApplePlatformSymbolication::SymbolInfoForStrippedSymbol(uint64 ProgramCounter, ANSICHAR const* ModulePath, ANSICHAR const* ModuleUUID, FProgramCounterSymbolInfo& Info)
+{
+#if PLATFORM_MAC
+	bool bOK = false;
+	
+	if(IFileManager::Get().FileSize(UTF8_TO_TCHAR(ModulePath)) > 0)
+	{
+		CSSymbolicatorRef Symbolicator = CSSymbolicatorCreateWithPathAndArchitecture(ModulePath, CPU_TYPE_X86_64);
+		if(!CSIsNull(Symbolicator))
+		{
+			FString ModuleID(ModuleUUID);
+			CFUUIDRef UUID = CFUUIDCreateFromString(nullptr, (CFStringRef)ModuleID.GetNSString());
+			check(UUID);
+			
+			CSSymbolOwnerRef SymbolOwner = CSSymbolicatorGetSymbolOwnerWithUUIDAtTime(Symbolicator, UUID, kCSNow);
+			if(!CSIsNull(SymbolOwner))
+			{
+				CSSourceInfoRef Symbol = CSSymbolicatorGetSourceInfoWithAddressAtTime(Symbolicator, (vm_address_t)ProgramCounter, kCSNow);
+				
+				if(!CSIsNull(Symbol))
+				{
+					Info.LineNumber = CSSourceInfoGetLineNumber(Symbol);
+					
+					CSRange CodeRange = CSSourceInfoGetRange(Symbol);
+					Info.SymbolDisplacement = (ProgramCounter - CodeRange.Location);
+					
+					ANSICHAR const* FileName = CSSourceInfoGetPath(Symbol);
+					if(FileName)
+					{
+						FCStringAnsi::Sprintf(Info.Filename, FileName);
+					}
+					
+					CSSymbolRef SymbolData = CSSourceInfoGetSymbol(Symbol);
+					if(!CSIsNull(SymbolData))
+					{
+						ANSICHAR const* FunctionName = CSSymbolGetName(SymbolData);
+						if(FunctionName)
+						{
+							FCStringAnsi::Sprintf(Info.FunctionName, FunctionName);
+						}
+					}
+					
+					CSSymbolOwnerRef Owner = CSSourceInfoGetSymbolOwner(Symbol);
+					if(!CSIsNull(Owner))
+					{
+						ANSICHAR const* DylibName = CSSymbolOwnerGetName(Owner);
+						FCStringAnsi::Strcpy(Info.ModuleName, DylibName);
+						
+						bOK = Info.LineNumber != 0;
+					}
+				}
+			}
+			
+			CFRelease(UUID);
+			CSRelease(Symbolicator);
+		}
+	}
+	
 	return bOK;
 #else
 	return false;
diff --git a/Engine/Source/Runtime/Core/Private/Mac/MacPlatformMisc.cpp b/Engine/Source/Runtime/Core/Private/Mac/MacPlatformMisc.cpp
--- a/Engine/Source/Runtime/Core/Private/Mac/MacPlatformMisc.cpp
+++ b/Engine/Source/Runtime/Core/Private/Mac/MacPlatformMisc.cpp
@@ -1166,106 +1166,6 @@ void FMacPlatformMisc::SetCrashHandler(void (* CrashHandler)(const FGenericCrash
 	sigaction(SIGABRT, &Action, NULL);
 }
 
-void FMacCrashContext::GenerateReport(char const* DiagnosticsPath) const
-{
-	int ReportFile = open(DiagnosticsPath, O_CREAT|O_WRONLY, 0766);
-	if (ReportFile != -1)
-	{
-		char Line[PATH_MAX] = {};
-		
-		WriteLine(ReportFile, "Generating report for minidump");
-		WriteLine(ReportFile);
-		
-		FCStringAnsi::Strncpy(Line, "Application version 4.0.", PATH_MAX);
-		FCStringAnsi::Strcat(Line, PATH_MAX, ItoANSI(ENGINE_VERSION_HIWORD, 10));
-		FCStringAnsi::Strcat(Line, PATH_MAX, ".");
-		FCStringAnsi::Strcat(Line, PATH_MAX, ItoANSI(ENGINE_VERSION_LOWORD, 10));
-		WriteLine(ReportFile, Line);
-		
-		FCStringAnsi::Strncpy(Line, " ... built from changelist ", PATH_MAX);
-		FCStringAnsi::Strcat(Line, PATH_MAX, ItoANSI(ENGINE_VERSION, 10));
-		WriteLine(ReportFile, Line);
-		WriteLine(ReportFile);
-		
-		FCStringAnsi::Strncpy(Line, "OS version Mac OS X ", PATH_MAX);
-		FCStringAnsi::Strcat(Line, PATH_MAX, GMacAppInfo.OSVersionUTF8);
-		FCStringAnsi::Strcat(Line, PATH_MAX, " (network name: ");
-		FCStringAnsi::Strcat(Line, PATH_MAX, GMacAppInfo.MachineName);
-		FCStringAnsi::Strcat(Line, PATH_MAX, ")");
-		WriteLine(ReportFile, Line);
-		
-		FCStringAnsi::Strncpy(Line, "Running ", PATH_MAX);
-		FCStringAnsi::Strcat(Line, PATH_MAX, ItoANSI(FPlatformMisc::NumberOfCores(), 10));
-		FCStringAnsi::Strcat(Line, PATH_MAX, " ");
-		FCStringAnsi::Strcat(Line, PATH_MAX, GMacAppInfo.MachineCPUString);
-		FCStringAnsi::Strcat(Line, PATH_MAX, "processors (");
-		FCStringAnsi::Strcat(Line, PATH_MAX, ItoANSI(FPlatformMisc::NumberOfCoresIncludingHyperthreads(), 10));
-		FCStringAnsi::Strcat(Line, PATH_MAX, " logical cores)");
-		WriteLine(ReportFile, Line);
-
-		FCStringAnsi::Strncpy(Line, "Exception was \"", PATH_MAX);
-		FCStringAnsi::Strcat(Line, PATH_MAX, SignalDescription);
-		FCStringAnsi::Strcat(Line, PATH_MAX, "\"");
-		WriteLine(ReportFile, Line);
-		WriteLine(ReportFile);
-		
-		WriteLine(ReportFile, "<SOURCE START>");
-		WriteLine(ReportFile, "<SOURCE END>");
-		WriteLine(ReportFile);
-		
-		WriteLine(ReportFile, "<CALLSTACK START>");
-		WriteLine(ReportFile, MinidumpCallstackInfo);
-		WriteLine(ReportFile, "<CALLSTACK END>");
-		WriteLine(ReportFile);
-		
-		// Technically _dyld_image_count & _dyld_get_image_name aren't async handler safe
-		// however, I imagine they actually are, since they merely access an internal list
-		// which isn't even thread safe...
-		uint32 ModuleCount = _dyld_image_count();
-		FCStringAnsi::Strncpy(Line, ItoANSI(ModuleCount, 10), PATH_MAX);
-		FCStringAnsi::Strcat(Line, PATH_MAX, " loaded modules");
-		WriteLine(ReportFile, Line);
-		WriteLine(ReportFile);
-		
-		WriteLine(ReportFile, "<MODULES START>");
-		for(uint32 Index = 0; Index < ModuleCount; Index++)
-		{
-			char const* ModulePath = _dyld_get_image_name(Index);
-			uint32 Version = 0;
-			struct mach_header_64* Header = (struct mach_header_64*)_dyld_get_image_header(Index);
-			struct load_command *CurrentCommand = (struct load_command *)( (char *)Header + sizeof(struct mach_header_64) );
-			if( Header->magic == MH_MAGIC_64 )
-			{
-				for( int32 i = 0; i < Header->ncmds; i++ )
-				{
-					if( CurrentCommand->cmd == LC_LOAD_DYLIB )
-					{
-						struct dylib_command *DylibCommand = (struct dylib_command *) CurrentCommand;
-						Version = DylibCommand->dylib.current_version;
-						Version = ((Version & 0xff) + ((Version >> 8) & 0xff) * 100 + ((Version >> 16) & 0xffff) * 10000);
-						break;
-					}
-					
-					CurrentCommand = (struct load_command *)( (char *)CurrentCommand + CurrentCommand->cmdsize );
-				}
-			}
-			
-			// Add version information (very useful!)
-			FCStringAnsi::Strncpy(Line, ModulePath, PATH_MAX);
-			FCStringAnsi::Strcat(Line, PATH_MAX, " (");
-			FCStringAnsi::Strcat(Line, PATH_MAX, ItoANSI(Version, 10));
-			FCStringAnsi::Strcat(Line, PATH_MAX, ")");
-			WriteLine(ReportFile, Line);
-		}
-		WriteLine(ReportFile, "<MODULES END>");
-		WriteLine(ReportFile);
-		
-		WriteLine(ReportFile, "Report end!");
-		
-		close(ReportFile);
-	}
-}
-
 void FMacCrashContext::GenerateWindowsErrorReport(char const* WERPath) const
 {
 	int ReportFile = open(WERPath, O_CREAT|O_WRONLY, 0766);
@@ -1457,16 +1357,258 @@ void FMacCrashContext::GenerateWindowsErrorReport(char const* WERPath) const
 	}
 }
 
-void FMacCrashContext::GenerateMinidump(char const* MinidumpCallstackInfo, char const* Path) const
+void FMacCrashContext::GenerateMinidump(char const* Path) const
 {
 	int ReportFile = open(Path, O_CREAT|O_WRONLY, 0766);
 	if (ReportFile != -1)
 	{
+		TCHAR Line[PATH_MAX] = {};
+		
 		// write BOM
 		static uint16 ByteOrderMarker = 0xFEFF;
 		write(ReportFile, &ByteOrderMarker, sizeof(ByteOrderMarker));
 		
-		WriteLine(ReportFile, ANSI_TO_TCHAR(MinidumpCallstackInfo));
+		WriteUTF16String(ReportFile, TEXT("Process:\t"));
+		WriteUTF16String(ReportFile, *GMacAppInfo.AppName);
+		WriteUTF16String(ReportFile, TEXT(" ["));
+		WriteUTF16String(ReportFile, ItoTCHAR(FPlatformProcess::GetCurrentProcessId(), 10));
+		WriteLine(ReportFile, TEXT("]"));
+		
+		WriteUTF16String(ReportFile, TEXT("Path:\t"));
+		WriteLine(ReportFile, *GMacAppInfo.AppPath);
+		
+		WriteUTF16String(ReportFile, TEXT("Identifier:\t"));
+		WriteLine(ReportFile, *GMacAppInfo.AppName);
+		
+		WriteUTF16String(ReportFile, TEXT("Version:\t"));
+		WriteUTF16String(ReportFile, ENGINE_VERSION_STRINGIFY(ENGINE_MAJOR_VERSION) ENGINE_VERSION_TEXT(".") ENGINE_VERSION_STRINGIFY(ENGINE_MINOR_VERSION) ENGINE_VERSION_TEXT(".") ENGINE_VERSION_STRINGIFY(ENGINE_PATCH_VERSION));
+		WriteUTF16String(ReportFile, TEXT(" ("));
+		WriteUTF16String(ReportFile, ENGINE_VERSION_STRING);
+		WriteLine(ReportFile, TEXT(")"));
+		
+		WriteLine(ReportFile, TEXT("Code Type:\tX86-64 (Native)"));
+		
+		WriteUTF16String(ReportFile, TEXT("Parent Process:\t"));
+		WriteUTF16String(ReportFile, *GMacAppInfo.ParentProcess);
+		WriteUTF16String(ReportFile, TEXT(" ["));
+		WriteUTF16String(ReportFile, ItoTCHAR(getppid(), 10));
+		WriteLine(ReportFile, TEXT("]"));
+		
+		WriteUTF16String(ReportFile, TEXT("Responsible:\t"));
+		WriteUTF16String(ReportFile, *GMacAppInfo.ParentProcess);
+		WriteUTF16String(ReportFile, TEXT(" ["));
+		WriteUTF16String(ReportFile, ItoTCHAR(getppid(), 10));
+		WriteLine(ReportFile, TEXT("]"));
+		
+		WriteUTF16String(ReportFile, TEXT("User ID:\t"));
+		WriteLine(ReportFile, ItoTCHAR(getuid(), 10));
+		
+		WriteLine(ReportFile, TEXT(""));
+		
+		WriteUTF16String(ReportFile, TEXT("Date/Time:\t"));
+		
+		time_t Now;
+		time(&Now);
+		struct tm LocalTime;
+		localtime_r(&Now, &LocalTime);
+		wcsftime(Line, PATH_MAX, TEXT("%Y-%m-%d %H:%M:%S %z"), &LocalTime);
+		WriteLine(ReportFile, Line);
+		
+		WriteUTF16String(ReportFile, TEXT("OS Version:\tMac OS X "));
+		WriteUTF16String(ReportFile, *GMacAppInfo.OSVersion);
+		WriteUTF16String(ReportFile, TEXT(" ("));
+		WriteUTF16String(ReportFile, *GMacAppInfo.OSBuild);
+		WriteLine(ReportFile, TEXT(")"));
+		
+		WriteLine(ReportFile, TEXT("Report Version:\t11"));
+		
+		WriteUTF16String(ReportFile, TEXT("Anonymous UUID:\t"));
+		WriteUTF16String(ReportFile, *GMacAppInfo.MachineUUID);
+										  
+		WriteLine(ReportFile, TEXT(""));
+		WriteLine(ReportFile, TEXT(""));
+		
+		WriteLine(ReportFile, TEXT("Crashed Thread:\t0"));
+		WriteLine(ReportFile, TEXT(""));
+		
+		WriteUTF16String(ReportFile, TEXT("Exception Type:\t"));
+		WriteUTF16String(ReportFile, TEXT(" ("));
+		WriteUTF16String(ReportFile, ItoTCHAR(Info->si_errno, 10));
+		WriteUTF16String(ReportFile, TEXT(", "));
+		WriteUTF16String(ReportFile, ItoTCHAR(Signal, 10));
+		WriteLine(ReportFile, TEXT(")"));
+		
+		WriteUTF16String(ReportFile, TEXT("Exception Codes:\t"));
+		WriteUTF16String(ReportFile, TEXT(" 0x"));
+		WriteUTF16String(ReportFile, ItoTCHAR(Info->si_code, 16));
+		WriteUTF16String(ReportFile, TEXT(", 0x"));
+		WriteLine(ReportFile, ItoTCHAR(Info->si_band, 16));
+		
+		WriteLine(ReportFile, TEXT("Application Specific Information:"));
+		for(uint32 i = 0; i < PATH_MAX && (i < FCStringAnsi::Strlen(SignalDescription) + 1); i++)
+		{
+			Line[i] = SignalDescription[i];
+		}
+		WriteLine(ReportFile, Line);
+		
+		WriteLine(ReportFile, TEXT(""));
+		WriteLine(ReportFile, TEXT("Thread 0:"));
+		
+		static const int MAX_DEPTH = 100;
+		static uint64 StackTrace[MAX_DEPTH];
+		FMemory::MemZero( StackTrace );
+		FPlatformStackWalk::CaptureStackBackTrace( StackTrace, MAX_DEPTH, Context );
+		
+		// Skip the first two entries as they are inside the stack walking code.
+		Dl_info Info;
+		static const int32 IgnoreCount = 6;
+		int32 CurrentDepth = IgnoreCount;
+		// Allow the first entry to be NULL as the crash could have been caused by a call to a NULL function pointer,
+		// which would mean the top of the callstack is NULL.
+		while( StackTrace[CurrentDepth] || ( CurrentDepth == IgnoreCount ) )
+		{
+			WriteUTF16String(ReportFile, ItoTCHAR(CurrentDepth - IgnoreCount, 10));
+			WriteUTF16String(ReportFile, TEXT("\t"));
+			
+			dladdr((void*)StackTrace[CurrentDepth], &Info);
+			if(Info.dli_fname && FCStringAnsi::Strrchr(Info.dli_fname, '/'))
+			{
+				char const* Name = FCStringAnsi::Strrchr(Info.dli_fname, '/');
+				for(uint32 i = 1; i < PATH_MAX && (i < FCStringAnsi::Strlen(Name) + 1); i++)
+				{
+					Line[i-1] = Name[i];
+				}
+				WriteUTF16String(ReportFile, Line);
+			}
+			else
+			{
+				WriteUTF16String(ReportFile, TEXT("Unknown Module"));
+			}
+			WriteUTF16String(ReportFile, TEXT("\t0x"));
+			WriteUTF16String(ReportFile, ItoTCHAR(StackTrace[CurrentDepth], 16));
+			WriteUTF16String(ReportFile, TEXT(" "));
+			if(Info.dli_sname && FCStringAnsi::Strlen(Info.dli_sname))
+			{
+				for(uint32 i = 0; i < PATH_MAX && (i < FCStringAnsi::Strlen(Info.dli_sname) + 1); i++)
+				{
+					Line[i] = Info.dli_sname[i];
+				}
+				WriteUTF16String(ReportFile, Line);
+			}
+			else
+			{
+				WriteUTF16String(ReportFile, TEXT("Unknown Function"));
+			}
+			WriteUTF16String(ReportFile, TEXT(" + "));
+			WriteLine(ReportFile, ItoTCHAR(StackTrace[CurrentDepth] - (uint64)Info.dli_saddr, 10));
+			
+			CurrentDepth++;
+		}
+		
+		WriteLine(ReportFile, TEXT(""));
+		
+		WriteLine(ReportFile, TEXT("Binary Images:"));
+		uint32 ModuleCount = _dyld_image_count();
+		for(uint32 Index = 0; Index < ModuleCount; Index++)
+		{
+			char const* ModulePath = _dyld_get_image_name(Index);
+			uint32 CompatVersion = 0;
+			uint32 CurrentVersion = 0;
+			uint8 UUID[16] = {0,0,0,0,0,0,0,0,0,0,0,0,0,0,0,0};
+			struct mach_header_64* Header = (struct mach_header_64*)_dyld_get_image_header(Index);
+			struct load_command *CurrentCommand = (struct load_command *)( (char *)Header + sizeof(struct mach_header_64) );
+			if( Header->magic == MH_MAGIC_64 )
+			{
+				uint8 Found = 0;
+				for( int32 i = 0; i < Header->ncmds && Found < 2; i++ )
+				{
+					if( CurrentCommand->cmd == LC_LOAD_DYLIB )
+					{
+						struct dylib_command *DylibCommand = (struct dylib_command *) CurrentCommand;
+						CompatVersion = DylibCommand->dylib.compatibility_version;
+						CurrentVersion = DylibCommand->dylib.current_version;
+						Found++;
+					}
+					else if ( CurrentCommand->cmd == LC_UUID )
+					{
+						struct uuid_command* UUIDCommand = (struct uuid_command*)CurrentCommand;
+						FMemory::Memcpy(UUID, UUIDCommand->uuid, 16);
+						Found++;
+					}
+					
+					CurrentCommand = (struct load_command *)( (char *)CurrentCommand + CurrentCommand->cmdsize );
+				}
+			}
+			
+			WriteUTF16String(ReportFile, TEXT("0x"));
+			WriteUTF16String(ReportFile, ItoTCHAR((uint64)Header, 16));
+			WriteUTF16String(ReportFile, TEXT(" -\t0x"));
+			WriteUTF16String(ReportFile, ItoTCHAR((uint64)Header + sizeof(struct mach_header_64) + Header->sizeofcmds, 16));
+			WriteUTF16String(ReportFile, TEXT(" "));
+			
+			char const* Name = FCStringAnsi::Strrchr(ModulePath, '/');
+			for(uint32 i = 1; i < PATH_MAX && (i < FCStringAnsi::Strlen(Name) + 1); i++)
+			{
+				Line[i-1] = Name[i];
+			}
+			WriteUTF16String(ReportFile, Line);
+			WriteUTF16String(ReportFile, TEXT(" ("));
+			
+			WriteUTF16String(ReportFile, ItoTCHAR(((CompatVersion >> 16) & 0xffff), 10));
+			WriteUTF16String(ReportFile, TEXT("."));
+			WriteUTF16String(ReportFile, ItoTCHAR(((CompatVersion >> 8) & 0xff), 10));
+			WriteUTF16String(ReportFile, TEXT("."));
+			WriteUTF16String(ReportFile, ItoTCHAR((CompatVersion & 0xff), 10));
+			WriteUTF16String(ReportFile, TEXT(" - "));
+			WriteUTF16String(ReportFile, ItoTCHAR(((CurrentVersion >> 16) & 0xffff), 10));
+			WriteUTF16String(ReportFile, TEXT("."));
+			WriteUTF16String(ReportFile, ItoTCHAR(((CurrentVersion >> 8) & 0xff), 10));
+			WriteUTF16String(ReportFile, TEXT("."));
+			WriteUTF16String(ReportFile, ItoTCHAR((CurrentVersion & 0xff), 10));
+			WriteUTF16String(ReportFile, TEXT(") <"));
+			WriteUTF16String(ReportFile, ItoTCHAR(UUID[0], 16, 2));
+			WriteUTF16String(ReportFile, ItoTCHAR(UUID[1], 16, 2));
+			WriteUTF16String(ReportFile, ItoTCHAR(UUID[2], 16, 2));
+			WriteUTF16String(ReportFile, ItoTCHAR(UUID[3], 16, 2));
+			WriteUTF16String(ReportFile, TEXT("-"));
+			WriteUTF16String(ReportFile, ItoTCHAR(UUID[4], 16, 2));
+			WriteUTF16String(ReportFile, ItoTCHAR(UUID[5], 16, 2));
+			WriteUTF16String(ReportFile, TEXT("-"));
+			WriteUTF16String(ReportFile, ItoTCHAR(UUID[6], 16, 2));
+			WriteUTF16String(ReportFile, ItoTCHAR(UUID[7], 16, 2));
+			WriteUTF16String(ReportFile, TEXT("-"));
+			WriteUTF16String(ReportFile, ItoTCHAR(UUID[8], 16, 2));
+			WriteUTF16String(ReportFile, ItoTCHAR(UUID[9], 16, 2));
+			WriteUTF16String(ReportFile, TEXT("-"));
+			WriteUTF16String(ReportFile, ItoTCHAR(UUID[10], 16, 2));
+			WriteUTF16String(ReportFile, ItoTCHAR(UUID[11], 16, 2));
+			WriteUTF16String(ReportFile, ItoTCHAR(UUID[12], 16, 2));
+			WriteUTF16String(ReportFile, ItoTCHAR(UUID[13], 16, 2));
+			WriteUTF16String(ReportFile, ItoTCHAR(UUID[14], 16, 2));
+			WriteUTF16String(ReportFile, ItoTCHAR(UUID[15], 16, 2));
+			WriteUTF16String(ReportFile, TEXT("> "));
+			
+			for(uint32 i = 0; i < PATH_MAX && (i < FCStringAnsi::Strlen(ModulePath) + 1); i++)
+			{
+				Line[i] = ModulePath[i];
+			}
+			WriteLine(ReportFile, Line);
+		}
+		
+		WriteLine(ReportFile, TEXT(""));
+		
+		WriteLine(ReportFile, TEXT("System Profile:"));
+		WriteUTF16String(ReportFile, TEXT("Model: "));
+		WriteUTF16String(ReportFile, *GMacAppInfo.MachineModel);
+		WriteUTF16String(ReportFile, TEXT(", "));
+		for(uint32 i = 0; i < PATH_MAX && (i < FCStringAnsi::Strlen(GMacAppInfo.MachineCPUString) + 1); i++)
+		{
+			Line[i] = GMacAppInfo.MachineCPUString[i];
+		}
+		WriteLine(ReportFile, Line);
+		
+		WriteUTF16String(ReportFile, TEXT("Graphics: "));
+		WriteLine(ReportFile, *GMacAppInfo.PrimaryGPU);
 		
 		close(ReportFile);
 	}
@@ -1505,20 +1647,15 @@ void FMacCrashContext::GenerateCrashInfoAndLaunchReporter() const
 			close(ReportFile);
 		}
 		
-		// generate "minidump"
-		FCStringAnsi::Strncpy(FilePath, CrashInfoFolder, PATH_MAX);
-		FCStringAnsi::Strcat(FilePath, PATH_MAX, "/diagnostics.txt");
-		GenerateReport(FilePath);
-		
 		// generate "WER"
 		FCStringAnsi::Strncpy(FilePath, CrashInfoFolder, PATH_MAX);
 		FCStringAnsi::Strcat(FilePath, PATH_MAX, "/wermeta.xml");
 		GenerateWindowsErrorReport(FilePath);
 		
-		// generate "minidump" (just >1 byte)
+		// generate "minidump" (Apple crash log format)
 		FCStringAnsi::Strncpy(FilePath, CrashInfoFolder, PATH_MAX);
 		FCStringAnsi::Strcat(FilePath, PATH_MAX, "/minidump.dmp");
-		GenerateMinidump(MinidumpCallstackInfo, FilePath);
+		GenerateMinidump(FilePath);
 		
 		// generate "info.txt" custom data for our server
 		FCStringAnsi::Strncpy(FilePath, CrashInfoFolder, PATH_MAX);
diff --git a/Engine/Source/Runtime/Core/Public/Apple/ApplePlatformCrashContext.h b/Engine/Source/Runtime/Core/Public/Apple/ApplePlatformCrashContext.h
--- a/Engine/Source/Runtime/Core/Public/Apple/ApplePlatformCrashContext.h
+++ b/Engine/Source/Runtime/Core/Public/Apple/ApplePlatformCrashContext.h
@@ -56,8 +56,8 @@ struct CORE_API FApplePlatformCrashContext : public FGenericCrashContext
 	static void WriteLine(int ReportFile, const TCHAR* Line);
 	
 	/** Async-safe ItoA */
-	static ANSICHAR* ItoANSI(uint64 Val, uint64 Base);
+	static ANSICHAR* ItoANSI(uint64 Val, uint64 Base, uint32 Len = 0);
 	
 	/** Async-safe ItoT */
-	static TCHAR* ItoTCHAR(uint64 Val, uint64 Base);
+	static TCHAR* ItoTCHAR(uint64 Val, uint64 Base, uint32 Len = 0);
 };
diff --git a/Engine/Source/Runtime/Core/Public/Apple/ApplePlatformSymbolication.h b/Engine/Source/Runtime/Core/Public/Apple/ApplePlatformSymbolication.h
--- a/Engine/Source/Runtime/Core/Public/Apple/ApplePlatformSymbolication.h
+++ b/Engine/Source/Runtime/Core/Public/Apple/ApplePlatformSymbolication.h
@@ -9,10 +9,11 @@
 #include "ApplePlatformStackWalk.h"
 
 /**
- * Apple platform implementation of symbolication
+ * Apple platform implementation of symbolication - not async. handler safe, so don't call during crash reporting!
  */
 struct CORE_API FApplePlatformSymbolication
 {
 	static bool SymbolInfoForAddress(uint64 ProgramCounter, FProgramCounterSymbolInfo& Info);
 	static bool SymbolInfoForFunctionFromModule(ANSICHAR const* MangledName, ANSICHAR const* ModuleName, FProgramCounterSymbolInfo& Info);
+	static bool SymbolInfoForStrippedSymbol(uint64 ProgramCounter, ANSICHAR const* ModulePath, ANSICHAR const* ModuleUUID, FProgramCounterSymbolInfo& Info);
 };
diff --git a/Engine/Source/Runtime/Core/Public/Mac/MacPlatformMisc.h b/Engine/Source/Runtime/Core/Public/Mac/MacPlatformMisc.h
--- a/Engine/Source/Runtime/Core/Public/Mac/MacPlatformMisc.h
+++ b/Engine/Source/Runtime/Core/Public/Mac/MacPlatformMisc.h
@@ -189,17 +189,11 @@ enum EMacModifierKeys
 
 struct CORE_API FMacCrashContext : public FApplePlatformCrashContext
 {
-	/**
-	 * Dumps all the data from crash context to the "minidump" report.
-	 * @param DiagnosticsPath Path to put the file to
-	 */
-	void GenerateReport(char const* DiagnosticsPath) const;
-	
 	/** Mimics Windows WER format */
 	void GenerateWindowsErrorReport(char const* WERPath) const;
 	
 	/** Creates (fake so far) minidump */
-	void GenerateMinidump(char const* MinidumpCallstackInfo, char const* Path) const;
+	void GenerateMinidump(char const* Path) const;
 	
 	/** Generates information for crash reporter */
 	void GenerateCrashInfoAndLaunchReporter() const;
