Amend Mac crash reporting to emit data into the FCrashInfo structure and output using FCrashInfo's GenerateReport function. Add necessary parsing code to hoist out PID, Exception Code, number of cores etc. Fix various bugs that were preventing the output from being displayed & remove a redundant comment from LaunchMac.

diff --git a/Engine/Source/Developer/CrashDebugHelper/Private/CrashDebugHelper.cpp b/Engine/Source/Developer/CrashDebugHelper/Private/CrashDebugHelper.cpp
--- a/Engine/Source/Developer/CrashDebugHelper/Private/CrashDebugHelper.cpp
+++ b/Engine/Source/Developer/CrashDebugHelper/Private/CrashDebugHelper.cpp
@@ -519,7 +519,7 @@ void FCrashInfo::WriteLine( FArchive* ReportFile, const ANSICHAR* Line )
 		ReportFile->Serialize( ( void* )Line, StringBytes );
 	}
 
-	ReportFile->Serialize( TCHAR_TO_UTF8( LINE_TERMINATOR ), 2 );
+	ReportFile->Serialize( TCHAR_TO_UTF8( LINE_TERMINATOR ), FCStringWide::Strlen(LINE_TERMINATOR) );
 }
 
 /** 
@@ -535,7 +535,12 @@ void FCrashInfo::GenerateReport( const FString& DiagnosticsPath )
 		WriteLine( ReportFile, TCHAR_TO_UTF8( TEXT( "Generating report for minidump" ) ) );
 		WriteLine( ReportFile );
 
-		if( Modules.Num() > 0 )
+		if ( ProductVersion.Len() > 0 )
+		{
+			Line = FString::Printf( TEXT( "Application version %s" ), *ProductVersion );
+			WriteLine( ReportFile, TCHAR_TO_UTF8( *Line ) );
+		}
+		else if( Modules.Num() > 0 )
 		{
 			Line = FString::Printf( TEXT( "Application version %d.%d.%d" ), Modules[0].Major, Modules[0].Minor, Modules[0].Patch );
 			WriteLine( ReportFile, TCHAR_TO_UTF8( *Line ) );
diff --git a/Engine/Source/Developer/CrashDebugHelper/Private/Mac/CrashDebugHelperMac.cpp b/Engine/Source/Developer/CrashDebugHelper/Private/Mac/CrashDebugHelperMac.cpp
--- a/Engine/Source/Developer/CrashDebugHelper/Private/Mac/CrashDebugHelperMac.cpp
+++ b/Engine/Source/Developer/CrashDebugHelper/Private/Mac/CrashDebugHelperMac.cpp
@@ -2,6 +2,7 @@
 
 #include "CrashDebugHelperPrivatePCH.h"
 #include "Database.h"
+#include "EngineVersion.h"
 #include "ApplePlatformSymbolication.h"
 #include <cxxabi.h>
 
@@ -39,18 +40,18 @@ static int32 ParseVersion(TCHAR const* CrashLog, int32& OutMajor, int32& OutMino
 	return Found;
 }
 
-static int32 ParseOS(TCHAR const* CrashLog, int32& OutMajor, int32& OutMinor, int32& OutPatch, int32& OutBuild)
+static int32 ParseOS(TCHAR const* CrashLog, uint16& OutMajor, uint16& OutMinor, uint16& OutPatch, uint16& OutBuild)
 {
 	int32 Found = 0;
 	TCHAR const* VersionLine = FCStringWide::Strstr(CrashLog, TEXT("OS Version:"));
 	if (VersionLine)
 	{
-		Found = swscanf(VersionLine, TEXT("%*s %*s Mac OS X %d.%d.%d (%xd)"), &OutMajor, &OutMinor, &OutPatch, &OutBuild);
+		Found = swscanf(VersionLine, TEXT("%*s %*s Mac OS X %hd.%hd.%hd (%hxd)"), &OutMajor, &OutMinor, &OutPatch, &OutBuild);
 	}
 	return Found;
 }
 
-static bool ParseModel(TCHAR const* CrashLog, FString& OutModelDetails)
+static bool ParseModel(TCHAR const* CrashLog, FString& OutModelDetails, uint32& OutProcessorNum)
 {
 	bool bFound = false;
 	TCHAR const* Line = FCStringWide::Strstr(CrashLog, TEXT("Model:"));
@@ -67,6 +68,26 @@ static bool ParseModel(TCHAR const* CrashLog, FString& OutModelDetails)
 		int32 Length = FMath::Min(256, (int32)((uintptr_t)(End - Line)));
 		OutModelDetails.Append(Line, Length);
 		
+		OutProcessorNum = 1;
+		int32 ProcessorPos = OutModelDetails.Find(TEXT(" processors"));
+		if( ProcessorPos != INDEX_NONE )
+		{
+			int32 NumStart = ProcessorPos;
+			while(NumStart && OutModelDetails[NumStart] != TEXT(','))
+			{
+				NumStart--;
+			}
+			if(NumStart >= 0 && OutModelDetails[NumStart] == TEXT(','))
+			{
+				NumStart += 2;
+				FString NumProc = OutModelDetails.Mid(NumStart, ProcessorPos-NumStart);
+				if(NumProc.IsNumeric())
+				{
+					TTypeFromString<uint32>::FromString(OutProcessorNum, *NumProc);
+				}
+			}
+		}
+		
 		bFound = true;
 	}
 	return bFound;
@@ -87,6 +108,7 @@ static int32 ParseGraphics(TCHAR const* CrashLog, FString& OutGPUDetails)
 		}
 		check(End);
 		
+		OutGPUDetails.Append(TEXT(", "));
 		int32 Length = FMath::Min((256 - Output), (int32)((uintptr_t)(End - Line)));
 		OutGPUDetails.Append(Line, Length);
 		
@@ -121,21 +143,88 @@ static int32 ParseError(TCHAR const* CrashLog, FString& OutErrorDetails)
 	return bFound;
 }
 
-static int32 ParseCrashedThread(TCHAR const* CrashLog, int32& OutThreadNumber)
+static int32 ParseExceptionCode(TCHAR const* CrashLog, uint32& OutExceptionCode)
+{
+	int32 Found = 0;
+	TCHAR const* Line = FCStringWide::Strstr(CrashLog, TEXT("Exception Type:"));
+	if(Line)
+	{
+		TCHAR Buffer[257] = {0};
+		Found = swscanf(Line, TEXT("%*s %*s %*d (%256ls)"), Buffer);
+		if(Found)
+		{
+			TCHAR* End = FCStringWide::Strchr(Buffer, TEXT(')'));
+			if(End)
+			{
+				*End = TEXT('\0');
+			}
+			if(FCStringWide::Strcmp(Buffer, TEXT("SIGQUIT")) == 0)
+			{
+				OutExceptionCode = SIGQUIT;
+			}
+			else if(FCStringWide::Strcmp(Buffer, TEXT("SIGILL")) == 0)
+			{
+				OutExceptionCode = SIGILL;
+			}
+			else if(FCStringWide::Strcmp(Buffer, TEXT("SIGEMT")) == 0)
+			{
+				OutExceptionCode = SIGEMT;
+			}
+			else if(FCStringWide::Strcmp(Buffer, TEXT("SIGFPE")) == 0)
+			{
+				OutExceptionCode = SIGFPE;
+			}
+			else if(FCStringWide::Strcmp(Buffer, TEXT("SIGBUS")) == 0)
+			{
+				OutExceptionCode = SIGBUS;
+			}
+			else if(FCStringWide::Strcmp(Buffer, TEXT("SIGSEGV")) == 0)
+			{
+				OutExceptionCode = SIGSEGV;
+			}
+			else if(FCStringWide::Strcmp(Buffer, TEXT("SIGSYS")) == 0)
+			{
+				OutExceptionCode = SIGSYS;
+			}
+			else if(FCStringWide::Strcmp(Buffer, TEXT("SIGABRT")) == 0)
+			{
+				OutExceptionCode = SIGABRT;
+			}
+			else
+			{
+				Found = swscanf(Buffer, TEXT("%u"), &OutExceptionCode);
+			}
+		}
+	}
+	return Found;
+}
+
+static int32 ParseCrashedThread(TCHAR const* CrashLog, uint32& OutThreadNumber)
 {
 	int32 Found = 0;
 	TCHAR const* Line = FCStringWide::Strstr(CrashLog, TEXT("Crashed Thread:"));
 	if (Line)
 	{
-		Found = swscanf(Line, TEXT("%*s %*s %d"), &OutThreadNumber);
+		Found = swscanf(Line, TEXT("%*s %*s %u"), &OutThreadNumber);
+	}
+	return Found;
+}
+
+static int32 ParseProcessID(TCHAR const* CrashLog, uint32& OutPID)
+{
+	int32 Found = 0;
+	TCHAR const* Line = FCStringWide::Strstr(CrashLog, TEXT("Process:"));
+	if (Line)
+	{
+		Found = swscanf(Line, TEXT("%*s %*s [%u]"), &OutPID);
 	}
 	return Found;
 }
 
-static TCHAR const* FindThreadStack(TCHAR const* CrashLog, int32 const ThreadNumber)
+static TCHAR const* FindThreadStack(TCHAR const* CrashLog, uint32 const ThreadNumber)
 {
 	int32 Found = 0;
-	FString Format = FString::Printf(TEXT("Thread %d"), ThreadNumber);
+	FString Format = FString::Printf(TEXT("Thread %u"), ThreadNumber);
 	TCHAR const* Line = FCStringWide::Strstr(CrashLog, *Format);
 	if (Line)
 	{
@@ -149,7 +238,7 @@ static TCHAR const* FindThreadStack(TCHAR const* CrashLog, int32 const ThreadNum
 static TCHAR const* FindCrashedThreadStack(TCHAR const* CrashLog)
 {
 	TCHAR const* Line = nullptr;
-	int32 ThreadNumber = 0;
+	uint32 ThreadNumber = 0;
 	int32 Found = ParseCrashedThread(CrashLog, ThreadNumber);
 	if(Found)
 	{
@@ -333,30 +422,6 @@ static bool ParseModuleLine(TCHAR const* ModuleLine, FCrashModuleInfo& OutModule
 	return bOK;
 }
 
-static int64 StringSize( const ANSICHAR* Line )
-{
-	int64 Size = 0;
-	if( Line != NULL )
-	{
-		while( *Line++ != 0 )
-		{
-			Size++;
-		}
-	}
-	return Size;
-}
-
-static void WriteLine( FArchive* ReportFile, const ANSICHAR* Line = "" )
-{
-	if( Line != NULL )
-	{
-		int64 StringBytes = StringSize( Line );
-		ReportFile->Serialize( ( void* )Line, StringBytes );
-	}
-	
-	ReportFile->Serialize( TCHAR_TO_UTF8( LINE_TERMINATOR ), 1 );
-}
-
 FCrashDebugHelperMac::FCrashDebugHelperMac()
 {
 }
@@ -442,180 +507,137 @@ bool FCrashDebugHelperMac::CreateMinidumpDiagnosticReport( const FString& InCras
 		int32 Result = ParseReportVersion(*CrashDump, ReportVersion);
 		if(Result == 1 && ReportVersion == 11)
 		{
-			FString DiagnosticsPath = FPaths::GetPath(InCrashDumpName) / TEXT("diagnostics.txt");
+			FString Error;
+			FString ModulePath;
+			FString ModuleName;
+			FString FunctionName;
+			FString FileName;
+			FString Branch;
+			FString Model;
+			FString Gpu;
+			
+			uint64 ProgramCounter = 0;
+			int32 Major = 0;
+			int32 Minor = 0;
+			int32 Build = 0;
+			int32 CLNumber = 0;
+			int32 LineNumber = 0;;
 			
-			FArchive* ReportFile = IFileManager::Get().CreateFileWriter( *DiagnosticsPath );
-			if( ReportFile != NULL )
+			int32 Result = ParseVersion(*CrashDump, Major, Minor, Build, CLNumber, Branch);
+			if(Result >= 3)
 			{
-				FString Error;
-				FString ModulePath;
-				FString ModuleName;
-				FString FunctionName;
-				FString FileName;
-				FString Branch;
-				FString Model;
-				FString Gpu;
-				FString Line;
-				
-				uint64 ProgramCounter = 0;
-				int32 Major = 0;
-				int32 Minor = 0;
-				int32 Build = 0;
-				int32 CLNumber = 0;
-				int32 OSMajor = 0;
-				int32 OSMinor = 0;
-				int32 OSPatch = 0;
-				int32 OSBuild = 0;
-				int32 LineNumber = 0;;
-				
-				WriteLine( ReportFile, TCHAR_TO_UTF8( TEXT( "Generating report for minidump" ) ) );
-				WriteLine( ReportFile );
-				
-				int32 Result = ParseVersion(*CrashDump, Major, Minor, Build, CLNumber, Branch);
-				if(Result >= 3)
+				CrashInfo.ProductVersion = FEngineVersion(Major, Minor, Build, CLNumber, Branch).ToString();
+			}
+			
+			if(Result >= 4)
+			{
+				CrashInfo.ChangelistBuiltFrom = CLNumber;
+			}
+			
+			if(Result == 5 && Branch.Len() > 0)
+			{
+				CrashInfo.LabelName = Branch;
+			}
+			
+			Result = ParseOS(*CrashDump, CrashInfo.SystemInfo.OSMajor, CrashInfo.SystemInfo.OSMinor, CrashInfo.SystemInfo.OSBuild, CrashInfo.SystemInfo.OSRevision);
+			check(Result == 4);
+			
+			CrashInfo.SystemInfo.ProcessorArchitecture = PA_X64;
+			
+			ParseModel(*CrashDump, Model, CrashInfo.SystemInfo.ProcessorCount);
+			ParseGraphics(*CrashDump, Gpu);
+			CrashInfo.SystemInfo.Report = Model + Gpu;
+			
+			Result = ParseError(*CrashDump, CrashInfo.Exception.ExceptionString);
+			check(Result == 1);
+			
+			Result = ParseProcessID(*CrashDump, CrashInfo.Exception.ProcessId);
+			check(Result == 1);
+			
+			Result = ParseCrashedThread(*CrashDump, CrashInfo.Exception.ThreadId);
+			check(Result == 1);
+			
+			Result = ParseExceptionCode(*CrashDump, CrashInfo.Exception.Code);
+			check(Result == 1);
+			
+			FCrashThreadInfo ThreadInfo;
+			ThreadInfo.ThreadId = CrashInfo.Exception.ThreadId;
+			ThreadInfo.SuspendCount = 0;
+			
+			// Parse modules now for symbolication - if we don't have the running process we need to symbolicate by UUID
+			TCHAR const* ModuleLine = FindModules(*CrashDump);
+			while(ModuleLine)
+			{
+				FCrashModuleInfo Module;
+				bool const bOK = ParseModuleLine(ModuleLine, Module);
+				if(bOK)
 				{
-					Line = FString::Printf( TEXT( "Application version %d.%d.%d" ), Major, Minor, Build );
-					WriteLine( ReportFile, TCHAR_TO_UTF8( *Line ) );
+					CrashInfo.Modules.Push(Module);
+					CrashInfo.ModuleNames.Push(FPaths::GetBaseFilename(Module.Name));
+					
+					ModuleLine = FCStringWide::Strchr(ModuleLine, TEXT('\n'));
+					check(ModuleLine);
+					ModuleLine += 1;
 				}
-				
-				if(Result >= 4)
+				else
 				{
-					Line = FString::Printf( TEXT( " ... built from changelist %d" ), CLNumber );
-					WriteLine( ReportFile, TCHAR_TO_UTF8( *Line ) );
+					ModuleLine = nullptr;
 				}
+			}
+			
+			TCHAR const* ThreadStackLine = FindCrashedThreadStack(*CrashDump);
+			while(ThreadStackLine)
+			{
+				Result = ParseThreadStackLine(ThreadStackLine, ModuleName, ProgramCounter, FunctionName, FileName, LineNumber);
 				
-				if(Result == 5 && Branch.Len() > 0)
+				// If we got the modulename & program counter but didn't parse the filename & linenumber we can resymbolise
+				if(Result > 1 && Result < 4)
 				{
-					Line = FString::Printf( TEXT( " ... based on label %s" ), *Branch );
-					WriteLine( ReportFile, TCHAR_TO_UTF8( *Line ) );
+					// Attempt to resymbolise using CoreSymbolication
+					Result += SymboliseStackInfo(CrashInfo.Modules, ModuleName, ProgramCounter, FunctionName, FileName, LineNumber);
 				}
-				WriteLine( ReportFile );
 				
-				Result = ParseOS(*CrashDump, OSMajor, OSMinor, OSPatch, OSBuild);
-				check(Result == 4);
-				Line = FString::Printf( TEXT( "OS version %d.%d.%d.%xd" ), OSMajor, OSMinor, OSPatch, OSBuild );
-				WriteLine( ReportFile, TCHAR_TO_UTF8( *Line ) );
-				
-				ParseModel(*CrashDump, Model);
-				ParseGraphics(*CrashDump, Gpu);
-				
-				Line = FString::Printf( TEXT( "Running %s %s" ), *Model, *Gpu );
-				WriteLine( ReportFile, TCHAR_TO_UTF8( *Line ) );
-				
-				Result = ParseError(*CrashDump, Error);
-				check(Result == 1);
-				
-				Line = FString::Printf( TEXT( "Exception was \"%s\"" ), *Error );
-				WriteLine( ReportFile, TCHAR_TO_UTF8( *Line ) );
-				WriteLine( ReportFile );
-				
-				// Parse modules now for symbolication - if we don't have the running process we need to symbolicate by UUID
-				TArray<FCrashModuleInfo> Modules;
-				TCHAR const* ModuleLine = FindModules(*CrashDump);
-				while(ModuleLine)
+				// Output in our format based on the fields we actually have
+				switch (Result)
 				{
-					FCrashModuleInfo Module;
-					bool const bOK = ParseModuleLine(ModuleLine, Module);
-					if(bOK)
-					{
-						Modules.Push(Module);
+					case 2:
+						CrashInfo.Exception.CallStackString.Push( FString::Printf( TEXT( "Unknown() Address = 0x%lx (filename not found) [in %s]" ), ProgramCounter, *ModuleName ) );
+						ThreadInfo.CallStack.Push(ProgramCounter);
 						
-						ModuleLine = FCStringWide::Strchr(ModuleLine, TEXT('\n'));
-						check(ModuleLine);
-						ModuleLine += 1;
-					}
-					else
-					{
-						ModuleLine = nullptr;
-					}
-				}
-				
-				Line = FString::Printf( TEXT( "<CALLSTACK START>" ) );
-				WriteLine( ReportFile, TCHAR_TO_UTF8( *Line ) );
-				
-				TCHAR const* ThreadStackLine = FindCrashedThreadStack(*CrashDump);
-				while(ThreadStackLine)
-				{
-					Result = ParseThreadStackLine(ThreadStackLine, ModuleName, ProgramCounter, FunctionName, FileName, LineNumber);
-					
-					// If we got the modulename & program counter but didn't parse the filename & linenumber we can resymbolise
-					if(Result > 1 && Result < 4)
-					{
-						// Attempt to resymbolise using CoreSymbolication
-						Result += SymboliseStackInfo(Modules, ModuleName, ProgramCounter, FunctionName, FileName, LineNumber);
-					}
-					
-					// Output in our format based on the fields we actually have
-					switch (Result)
-					{
-						case 2:
-							Line = FString::Printf( TEXT( "Unknown() Address = 0x%lx (filename not found) [in %s]" ), ProgramCounter, *ModuleName );
-							WriteLine( ReportFile, TCHAR_TO_UTF8( *Line ) );
-							
-							ThreadStackLine = FCStringWide::Strchr(ThreadStackLine, TEXT('\n'));
-							check(ThreadStackLine);
-							ThreadStackLine += 1;
-							break;
-							
-						case 3:
-						case 4:
-							Line = FString::Printf( TEXT( "%s Address = 0x%lx (filename not found) [in %s]" ), *FunctionName, ProgramCounter, *ModuleName );
-							WriteLine( ReportFile, TCHAR_TO_UTF8( *Line ) );
-							
-							ThreadStackLine = FCStringWide::Strchr(ThreadStackLine, TEXT('\n'));
-							check(ThreadStackLine);
-							ThreadStackLine += 1;
-							break;
-							
-						case 5:
-						case 6: // Function name might be parsed twice
-							Line = FString::Printf( TEXT( "%s Address = 0x%lx [%s, line %d] [in %s]" ), *FunctionName, ProgramCounter, *FileName, LineNumber, *ModuleName );
-							WriteLine( ReportFile, TCHAR_TO_UTF8( *Line ) );
-							
-							ThreadStackLine = FCStringWide::Strchr(ThreadStackLine, TEXT('\n'));
-							check(ThreadStackLine);
-							ThreadStackLine += 1;
-							break;
-							
-						default:
-							ThreadStackLine = nullptr;
-							break;
-					}
-				}
-				
-				Line = FString::Printf( TEXT( "<CALLSTACK END>" ) );
-				WriteLine( ReportFile, TCHAR_TO_UTF8( *Line ) );
-				WriteLine( ReportFile );
-				
-				Line = FString::Printf( TEXT( "%d loaded modules" ), Modules.Num() );
-				WriteLine( ReportFile, TCHAR_TO_UTF8( *Line ) );
-				
-				for( int32 ModuleIndex = 0; ModuleIndex < Modules.Num(); ModuleIndex++ )
-				{
-					FCrashModuleInfo& Module = Modules[ModuleIndex];
-					
-					FString ModuleDirectory = FPaths::GetPath(Module.Name);
-					FString ModuleName = FPaths::GetBaseFilename( Module.Name, true ) + FPaths::GetExtension( Module.Name, true );
-					
-					FString ModuleDetail = FString::Printf( TEXT( "%40s" ), *ModuleName );
-					FString Version = FString::Printf( TEXT( " (%d.%d.%d.%d)" ), Module.Major, Module.Minor, Module.Patch, Module.Revision );
-					ModuleDetail += FString::Printf( TEXT( " %22s" ), *Version );
-					ModuleDetail += FString::Printf( TEXT( " 0x%016x 0x%08x" ), Module.BaseOfImage, Module.SizeOfImage );
-					ModuleDetail += FString::Printf( TEXT( " %s" ), *ModuleDirectory );
-					
-					WriteLine( ReportFile, TCHAR_TO_UTF8( *ModuleDetail ) );
+						ThreadStackLine = FCStringWide::Strchr(ThreadStackLine, TEXT('\n'));
+						check(ThreadStackLine);
+						ThreadStackLine += 1;
+						break;
+						
+					case 3:
+					case 4:
+						CrashInfo.Exception.CallStackString.Push( FString::Printf( TEXT( "%s Address = 0x%lx (filename not found) [in %s]" ), *FunctionName, ProgramCounter, *ModuleName ) );
+						ThreadInfo.CallStack.Push(ProgramCounter);
+						
+						ThreadStackLine = FCStringWide::Strchr(ThreadStackLine, TEXT('\n'));
+						check(ThreadStackLine);
+						ThreadStackLine += 1;
+						break;
+						
+					case 5:
+					case 6: // Function name might be parsed twice
+						CrashInfo.Exception.CallStackString.Push( FString::Printf( TEXT( "%s Address = 0x%lx [%s, line %d] [in %s]" ), *FunctionName, ProgramCounter, *FileName, LineNumber, *ModuleName ) );
+						ThreadInfo.CallStack.Push(ProgramCounter);
+						
+						ThreadStackLine = FCStringWide::Strchr(ThreadStackLine, TEXT('\n'));
+						check(ThreadStackLine);
+						ThreadStackLine += 1;
+						break;
+						
+					default:
+						ThreadStackLine = nullptr;
+						break;
 				}
-				
-				WriteLine( ReportFile );
-				
-				Line = FString::Printf( TEXT( "Report end!" ) );
-				WriteLine( ReportFile, TCHAR_TO_UTF8( *Line )  );
-				
-				ReportFile->Close();
-				delete ReportFile;
-				
-				bOK = true;
 			}
+			
+			CrashInfo.Threads.Push(ThreadInfo);
+			
+			bOK = true;
 		}
 	}
 	
diff --git a/Engine/Source/Programs/CrashReporter/CrashReportClient/Private/Mac/MacErrorReport.cpp b/Engine/Source/Programs/CrashReporter/CrashReportClient/Private/Mac/MacErrorReport.cpp
--- a/Engine/Source/Programs/CrashReporter/CrashReportClient/Private/Mac/MacErrorReport.cpp
+++ b/Engine/Source/Programs/CrashReporter/CrashReportClient/Private/Mac/MacErrorReport.cpp
@@ -114,7 +114,7 @@ FText FMacErrorReport::DiagnoseReport() const
 		return FText::FromString("No minidump found for this crash.");
 	}
 	
-	if (!CrashDebugHelper->CreateMinidumpDiagnosticReport(ReportDirectory / DumpFilename))
+	if ( !CrashDebugHelper->CreateMinidumpDiagnosticReport(ReportDirectory / DumpFilename) )
 	{
 		if ( FRocketSupport::IsRocket() )
 		{
@@ -128,7 +128,9 @@ FText FMacErrorReport::DiagnoseReport() const
 	else
 	{
 		FString CrashDump;
-		if ( FFileHelper::LoadFileToString( CrashDump, *(ReportDirectory / TEXT("diagnostics.txt")) ) )
+		FString DiagnosticsPath = ReportDirectory / FString( TEXT( "Diagnostics.txt" ) );
+		CrashDebugHelper->CrashInfo.GenerateReport( DiagnosticsPath );
+		if ( FFileHelper::LoadFileToString( CrashDump, *(ReportDirectory / TEXT("Diagnostics.txt")) ) )
 		{
 			return FText::FromString(CrashDump);
 		}
diff --git a/Engine/Source/Runtime/Core/Private/Mac/MacPlatformMisc.cpp b/Engine/Source/Runtime/Core/Private/Mac/MacPlatformMisc.cpp
--- a/Engine/Source/Runtime/Core/Private/Mac/MacPlatformMisc.cpp
+++ b/Engine/Source/Runtime/Core/Private/Mac/MacPlatformMisc.cpp
@@ -166,11 +166,14 @@ struct MacApplicationInfo
 		
 		uint32 Status = notify_register_dispatch(kIOPSNotifyPowerSource, &PowerSourceNotification, dispatch_get_main_queue(), PowerSourceNotifyHandler);
 		check(Status == NOTIFY_STATUS_OK);
+		
+		NumCores = FPlatformMisc::NumberOfCores();
 	}
 	
 	bool RunningOnBattery;
 	bool RunningOnMavericks;
 	int32 PowerSourceNotification;
+	int32 NumCores;
 	char AppNameUTF8[PATH_MAX+1];
 	char AppLogPath[PATH_MAX+1];
 	char CrashReportPath[PATH_MAX+1];
@@ -1427,9 +1430,8 @@ void FMacCrashContext::GenerateMinidump(char const* Path) const
 		WriteLine(ReportFile, TEXT(""));
 		
 		WriteUTF16String(ReportFile, TEXT("Exception Type:\t"));
-		WriteUTF16String(ReportFile, TEXT(" ("));
 		WriteUTF16String(ReportFile, ItoTCHAR(Info->si_errno, 10));
-		WriteUTF16String(ReportFile, TEXT(", "));
+		WriteUTF16String(ReportFile, TEXT(" ("));
 		WriteUTF16String(ReportFile, ItoTCHAR(Signal, 10));
 		WriteLine(ReportFile, TEXT(")"));
 		
@@ -1440,6 +1442,7 @@ void FMacCrashContext::GenerateMinidump(char const* Path) const
 		WriteLine(ReportFile, ItoTCHAR(Info->si_band, 16));
 		
 		WriteLine(ReportFile, TEXT("Application Specific Information:"));
+		FMemory::Memzero(Line, PATH_MAX * sizeof(TCHAR));
 		FUTF8ToTCHAR_Convert::Convert(Line, PATH_MAX, SignalDescription, FCStringAnsi::Strlen(SignalDescription));
 		WriteLine(ReportFile, Line);
 		
@@ -1472,6 +1475,8 @@ void FMacCrashContext::GenerateMinidump(char const* Path) const
 			if(Info.dli_fname && FCStringAnsi::Strrchr(Info.dli_fname, '/'))
 			{
 				char const* Name = FCStringAnsi::Strrchr(Info.dli_fname, '/');
+				Name++;
+				FMemory::Memzero(Line, PATH_MAX * sizeof(TCHAR));
 				FUTF8ToTCHAR_Convert::Convert(Line, PATH_MAX, Name, FCStringAnsi::Strlen(Name));
 				WriteUTF16String(ReportFile, Line);
 			}
@@ -1484,6 +1489,7 @@ void FMacCrashContext::GenerateMinidump(char const* Path) const
 			WriteUTF16String(ReportFile, TEXT(" "));
 			if(Info.dli_sname && FCStringAnsi::Strlen(Info.dli_sname))
 			{
+				FMemory::Memzero(Line, PATH_MAX * sizeof(TCHAR));
 				FUTF8ToTCHAR_Convert::Convert(Line, PATH_MAX, Info.dli_sname, FCStringAnsi::Strlen(Info.dli_sname));
 				WriteUTF16String(ReportFile, Line);
 			}
@@ -1539,6 +1545,8 @@ void FMacCrashContext::GenerateMinidump(char const* Path) const
 			WriteUTF16String(ReportFile, TEXT(" "));
 			
 			char const* Name = FCStringAnsi::Strrchr(ModulePath, '/');
+			Name++;
+			FMemory::Memzero(Line, PATH_MAX * sizeof(TCHAR));
 			FUTF8ToTCHAR_Convert::Convert(Line, PATH_MAX, Name, FCStringAnsi::Strlen(Name));
 			WriteUTF16String(ReportFile, Line);
 			WriteUTF16String(ReportFile, TEXT(" ("));
@@ -1577,6 +1585,7 @@ void FMacCrashContext::GenerateMinidump(char const* Path) const
 			WriteUTF16String(ReportFile, ItoTCHAR(UUID[15], 16, 2));
 			WriteUTF16String(ReportFile, TEXT("> "));
 			
+			FMemory::Memzero(Line, PATH_MAX * sizeof(TCHAR));
 			FUTF8ToTCHAR_Convert::Convert(Line, PATH_MAX, ModulePath, FCStringAnsi::Strlen(ModulePath));
 			WriteLine(ReportFile, Line);
 		}
@@ -1587,6 +1596,12 @@ void FMacCrashContext::GenerateMinidump(char const* Path) const
 		WriteUTF16String(ReportFile, TEXT("Model: "));
 		WriteUTF16String(ReportFile, *GMacAppInfo.MachineModel);
 		WriteUTF16String(ReportFile, TEXT(", "));
+		WriteUTF16String(ReportFile, ItoTCHAR(GMacAppInfo.NumCores, 10));
+		WriteUTF16String(ReportFile, TEXT(" processors, "));
+		WriteUTF16String(ReportFile, ItoTCHAR(GMacAppInfo.NumCores, 10));
+		WriteUTF16String(ReportFile, TEXT("-Core "));
+
+		FMemory::Memzero(Line, PATH_MAX * sizeof(TCHAR));
 		FUTF8ToTCHAR_Convert::Convert(Line, PATH_MAX, GMacAppInfo.MachineCPUString, FCStringAnsi::Strlen(GMacAppInfo.MachineCPUString));
 		WriteLine(ReportFile, Line);
 		
diff --git a/Engine/Source/Runtime/Launch/Private/Mac/LaunchMac.cpp b/Engine/Source/Runtime/Launch/Private/Mac/LaunchMac.cpp
--- a/Engine/Source/Runtime/Launch/Private/Mac/LaunchMac.cpp
+++ b/Engine/Source/Runtime/Launch/Private/Mac/LaunchMac.cpp
@@ -33,8 +33,6 @@ void EngineCrashHandler(const FGenericCrashContext & GenericContext)
 	{
 		GWarn->Flush();
 	}
-	// @todo Disable calls to functions which eventually try and allocated memory, resulting in the application hanging before the crash reporter can launch
-	// We need to setup the memory allocator system to handle the case where we can't call the system malloc, free etc. to re-enable this.
 	if (GError)
 	{
 		GError->Flush();
