Cache the symbolisation structures when symbolising crash reports on Mac OS X to speed the process up a bit. Only makes sense in the crash reporting case, since loading symbols from a running process is already very quick.

diff --git a/Engine/Source/Developer/CrashDebugHelper/Private/Mac/CrashDebugHelperMac.cpp b/Engine/Source/Developer/CrashDebugHelper/Private/Mac/CrashDebugHelperMac.cpp
--- a/Engine/Source/Developer/CrashDebugHelper/Private/Mac/CrashDebugHelperMac.cpp
+++ b/Engine/Source/Developer/CrashDebugHelper/Private/Mac/CrashDebugHelperMac.cpp
@@ -295,7 +295,7 @@ static int32 ParseThreadStackLine(TCHAR const* StackLine, FString& OutModuleName
 	return Found;
 }
 
-static int32 SymboliseStackInfo(TArray<FCrashModuleInfo> const& ModuleInfo, FString ModuleName, uint64 const ProgramCounter, FString& OutFunctionName, FString& OutFileName, int32& OutLineNumber)
+static int32 SymboliseStackInfo(FApplePlatformSymbolCache* SymbolCache, TArray<FCrashModuleInfo> const& ModuleInfo, FString ModuleName, uint64 const ProgramCounter, FString& OutFunctionName, FString& OutFileName, int32& OutLineNumber)
 {
 	FProgramCounterSymbolInfo Info;
 	int32 ValuesSymbolised = 0;
@@ -310,7 +310,7 @@ static int32 SymboliseStackInfo(TArray<FCrashModuleInfo> const& ModuleInfo, FStr
 		}
 	}
 
-	if((Module.Name.Len() > 0) && FApplePlatformSymbolication::SymbolInfoForStrippedSymbol((ProgramCounter - Module.BaseOfImage), TCHAR_TO_UTF8(*Module.Name), TCHAR_TO_UTF8(*Module.Report), Info))
+	if((Module.Name.Len() > 0) && FApplePlatformSymbolication::SymbolInfoForStrippedSymbol(SymbolCache, (ProgramCounter - Module.BaseOfImage), TCHAR_TO_UTF8(*Module.Name), TCHAR_TO_UTF8(*Module.Report), Info))
 	{
 		if(FCStringAnsi::Strlen(Info.FunctionName) > 0)
 		{
@@ -585,6 +585,8 @@ bool FCrashDebugHelperMac::CreateMinidumpDiagnosticReport( const FString& InCras
 				}
 			}
 			
+			FApplePlatformSymbolCache* SymbolCache = FApplePlatformSymbolication::CreateSymbolCache();
+			
 			TCHAR const* ThreadStackLine = FindCrashedThreadStack(*CrashDump);
 			while(ThreadStackLine)
 			{
@@ -594,7 +596,7 @@ bool FCrashDebugHelperMac::CreateMinidumpDiagnosticReport( const FString& InCras
 				if(Result > 1 && Result < 4)
 				{
 					// Attempt to resymbolise using CoreSymbolication
-					Result += SymboliseStackInfo(CrashInfo.Modules, ModuleName, ProgramCounter, FunctionName, FileName, LineNumber);
+					Result += SymboliseStackInfo(SymbolCache, CrashInfo.Modules, ModuleName, ProgramCounter, FunctionName, FileName, LineNumber);
 				}
 				
 				// Output in our format based on the fields we actually have
@@ -635,6 +637,8 @@ bool FCrashDebugHelperMac::CreateMinidumpDiagnosticReport( const FString& InCras
 				}
 			}
 			
+			FApplePlatformSymbolication::DestroySymbolCache(SymbolCache);
+			
 			CrashInfo.Threads.Push(ThreadInfo);
 			
 			bOK = true;
diff --git a/Engine/Source/Runtime/Core/Private/Apple/ApplePlatformSymbolication.cpp b/Engine/Source/Runtime/Core/Private/Apple/ApplePlatformSymbolication.cpp
--- a/Engine/Source/Runtime/Core/Private/Apple/ApplePlatformSymbolication.cpp
+++ b/Engine/Source/Runtime/Core/Private/Apple/ApplePlatformSymbolication.cpp
@@ -59,6 +59,36 @@ extern "C"
 }
 #endif
 
+struct FApplePlatformSymbolCache
+{
+#if PLATFORM_MAC
+	TMap< FString, CSSymbolicatorRef > Symbolicators;
+#endif
+};
+
+FApplePlatformSymbolCache* FApplePlatformSymbolication::CreateSymbolCache(void)
+{
+#if PLATFORM_MAC
+	return new FApplePlatformSymbolCache;
+#else 
+	return nullptr;
+#endif
+}
+
+void FApplePlatformSymbolication::DestroySymbolCache(FApplePlatformSymbolCache* Cache)
+{
+#if PLATFORM_MAC
+	if(Cache)
+	{
+		for(auto It : Cache->Symbolicators)
+		{
+			CSRelease( It.Value );
+		}
+		delete Cache;
+	}
+#endif
+}
+
 bool FApplePlatformSymbolication::SymbolInfoForAddress(uint64 ProgramCounter, FProgramCounterSymbolInfo& Info)
 {
 #if PLATFORM_MAC
@@ -122,16 +152,29 @@ bool FApplePlatformSymbolication::SymbolInfoForFunctionFromModule(ANSICHAR const
 #endif
 }
 
-bool FApplePlatformSymbolication::SymbolInfoForStrippedSymbol(uint64 ProgramCounter, ANSICHAR const* ModulePath, ANSICHAR const* ModuleUUID, FProgramCounterSymbolInfo& Info)
+bool FApplePlatformSymbolication::SymbolInfoForStrippedSymbol(FApplePlatformSymbolCache* Cache, uint64 ProgramCounter, ANSICHAR const* ModulePath, ANSICHAR const* ModuleUUID, FProgramCounterSymbolInfo& Info)
 {
 #if PLATFORM_MAC
 	bool bOK = false;
 	
 	if(IFileManager::Get().FileSize(UTF8_TO_TCHAR(ModulePath)) > 0)
 	{
-		CSSymbolicatorRef Symbolicator = CSSymbolicatorCreateWithPathAndArchitecture(ModulePath, CPU_TYPE_X86_64);
-		if(!CSIsNull(Symbolicator))
+		CSSymbolicatorRef Symbolicator = { nullptr, nullptr };
+		if ( Cache )
+		{
+			Symbolicator = Cache->Symbolicators.FindOrAdd( FString(ModulePath) );
+		}
+		if( CSIsNull(Symbolicator) )
+		{
+			Symbolicator = CSSymbolicatorCreateWithPathAndArchitecture(ModulePath, CPU_TYPE_X86_64);
+		}
+		if( !CSIsNull(Symbolicator) )
 		{
+			if ( Cache )
+			{
+				Cache->Symbolicators[ FString(ModulePath) ] = Symbolicator;
+			}
+			
 			FString ModuleID(ModuleUUID);
 			CFUUIDRef UUID = CFUUIDCreateFromString(nullptr, (CFStringRef)ModuleID.GetNSString());
 			check(UUID);
@@ -176,7 +219,11 @@ bool FApplePlatformSymbolication::SymbolInfoForStrippedSymbol(uint64 ProgramCoun
 			}
 			
 			CFRelease(UUID);
-			CSRelease(Symbolicator);
+			
+			if ( !Cache )
+			{
+				CSRelease(Symbolicator);
+			}
 		}
 	}
 	
diff --git a/Engine/Source/Runtime/Core/Public/Apple/ApplePlatformSymbolication.h b/Engine/Source/Runtime/Core/Public/Apple/ApplePlatformSymbolication.h
--- a/Engine/Source/Runtime/Core/Public/Apple/ApplePlatformSymbolication.h
+++ b/Engine/Source/Runtime/Core/Public/Apple/ApplePlatformSymbolication.h
@@ -8,6 +8,11 @@
 
 #include "ApplePlatformStackWalk.h"
 
+/**
+ * Opaque symbol cache for improved symbolisation performance.
+ */
+struct FApplePlatformSymbolCache;
+
 /**
  * Apple platform implementation of symbolication - not async. handler safe, so don't call during crash reporting!
  */
@@ -15,5 +20,8 @@ struct CORE_API FApplePlatformSymbolication
 {
 	static bool SymbolInfoForAddress(uint64 ProgramCounter, FProgramCounterSymbolInfo& Info);
 	static bool SymbolInfoForFunctionFromModule(ANSICHAR const* MangledName, ANSICHAR const* ModuleName, FProgramCounterSymbolInfo& Info);
-	static bool SymbolInfoForStrippedSymbol(uint64 ProgramCounter, ANSICHAR const* ModulePath, ANSICHAR const* ModuleUUID, FProgramCounterSymbolInfo& Info);
+	
+	static FApplePlatformSymbolCache* CreateSymbolCache(void);
+	static bool SymbolInfoForStrippedSymbol(FApplePlatformSymbolCache* Cache, uint64 ProgramCounter, ANSICHAR const* ModulePath, ANSICHAR const* ModuleUUID, FProgramCounterSymbolInfo& Info);
+	static void DestroySymbolCache(FApplePlatformSymbolCache* Cache);
 };
