Fix parsing of initial OS version string in Mac OS X (i.e. 10.10).

diff --git a/Engine/Source/Developer/CrashDebugHelper/Private/Mac/CrashDebugHelperMac.cpp b/Engine/Source/Developer/CrashDebugHelper/Private/Mac/CrashDebugHelperMac.cpp
--- a/Engine/Source/Developer/CrashDebugHelper/Private/Mac/CrashDebugHelperMac.cpp
+++ b/Engine/Source/Developer/CrashDebugHelper/Private/Mac/CrashDebugHelperMac.cpp
@@ -47,6 +47,11 @@ static int32 ParseOS(TCHAR const* CrashLog, uint16& OutMajor, uint16& OutMinor,
 	if (VersionLine)
 	{
 		Found = swscanf(VersionLine, TEXT("%*s %*s Mac OS X %hd.%hd.%hd (%hxd)"), &OutMajor, &OutMinor, &OutPatch, &OutBuild);
+        if ( Found == 2 )
+        {
+            OutPatch = 0;
+            Found += swscanf(VersionLine, TEXT("%*s %*s Mac OS X %*hd.%*hd (%hxd)"), &OutBuild) + 1;
+        }
 	}
 	return Found;
 }
