CrashReport - Tweaked PDB cache configuration to use less HDD space

diff --git a/Engine/Config/BaseEngine.ini b/Engine/Config/BaseEngine.ini
--- a/Engine/Config/BaseEngine.ini
+++ b/Engine/Config/BaseEngine.ini
@@ -2148,9 +2148,9 @@ ExecutablePathPattern=\\epicgames.net\root\Builds\Rocket\Automated\%PRODUCT_VERS
 ; devweb-02
 PDBCachePath=F:/CrashReportPDBCache/
 DepotRoot=F:/depot
-DaysToDeleteUnusedFilesFromPDBCache=14
+DaysToDeleteUnusedFilesFromPDBCache=7
 PDBCacheSizeGB=128
-MinDiskFreeSpaceGB=64
+MinDiskFreeSpaceGB=256
 
 [Engine.ErrorHandling]
 bPromptForRemoteDebugging=false
diff --git a/Engine/Source/Developer/CrashDebugHelper/Private/CrashDebugHelper.cpp b/Engine/Source/Developer/CrashDebugHelper/Private/CrashDebugHelper.cpp
--- a/Engine/Source/Developer/CrashDebugHelper/Private/CrashDebugHelper.cpp
+++ b/Engine/Source/Developer/CrashDebugHelper/Private/CrashDebugHelper.cpp
@@ -852,10 +852,11 @@ void FPDBCache::Init()
 {
 	// PDB Cache
 	// Default configuration
-	// PDBCachePath=F:/CrashReportPDBCache
-	// DaysToDeleteUnusedFilesFromPDBCache=14
-	// PDBCacheSizeGB=128
-	// MinFreeSizeGB=64
+	//PDBCachePath=F:/CrashReportPDBCache/
+	//DepotRoot=F:/depot
+	//DaysToDeleteUnusedFilesFromPDBCache=7
+	//PDBCacheSizeGB=128
+	//MinDiskFreeSpaceGB=256
 
 	// Can be enabled only through the command line.
 	FParse::Bool( FCommandLine::Get(), TEXT( "bUsePDBCache=" ), bUsePDBCache );
