CrashReport - Optimized FWindowsPlatformStackWalkExt::SetSymbolPathsFromModules, should be faster, removed recursion, optimized GetCallstacks to not use stack memory, now that code can be executed in thread pool where stack size is 32kb

diff --git a/Engine/Source/Developer/CrashDebugHelper/Private/Windows/WindowsPlatformStackWalkExt.cpp b/Engine/Source/Developer/CrashDebugHelper/Private/Windows/WindowsPlatformStackWalkExt.cpp
--- a/Engine/Source/Developer/CrashDebugHelper/Private/Windows/WindowsPlatformStackWalkExt.cpp
+++ b/Engine/Source/Developer/CrashDebugHelper/Private/Windows/WindowsPlatformStackWalkExt.cpp
@@ -192,6 +192,9 @@ private:
 
 public:
 
+	/** A list of directories to be checked, to avoid using recursive algorithm and save stack-memory. */
+	TArray<FString> RecursiveDirectories;
+
 	FSymbolAndExecutableFinder( IPlatformFile& InFileInterface, const TArray<FString>& InFiletypesToInclude, const TArray<FString>& InDirectoriesToIgnore )
 		: FileInterface( InFileInterface )
 		, FiletypesToInclude( InFiletypesToInclude )
@@ -232,7 +235,7 @@ public:
 			// Look in all the ignore directories looking for a match.
 			for( const auto& Directory : DirectoriesToIgnore )
 			{
-				if( RelativeFilename.Contains( Directory ) )
+				if( RelativeFilename.EndsWith( Directory ) )
 				{
 					bShouldRecurse = false;
 					break;
@@ -242,7 +245,7 @@ public:
 			// Recurse if we should.
 			if( bShouldRecurse )
 			{
-				FileInterface.IterateDirectory( FilenameOrDirectory, *this );
+				RecursiveDirectories.Add(FilenameOrDirectory);
 			}
 		}
 
@@ -300,16 +303,34 @@ void FWindowsPlatformStackWalkExt::SetSymbolPathsFromModules()
 		FiletypesToInclude.Add( TEXT( ".exe" ) );
 		// These folders don't contain any binary files, so skip them to improve the finder performance.
 		TArray<FString> DirectoriesToIgnore;
-		DirectoriesToIgnore.Add( TEXT( "/Content/" ) );
-		DirectoriesToIgnore.Add( TEXT( "/DerivedDataCache/" ) );
-		DirectoriesToIgnore.Add( TEXT( "/Intermediate/" ) );
-		DirectoriesToIgnore.Add( TEXT( "/Documentation/" ) );
-		DirectoriesToIgnore.Add( TEXT( "/NoRedist/" ) );
-
-
-		FSymbolAndExecutableFinder SymbolAndExecutableFinder( FPlatformFileManager::Get().GetPlatformFile(), FiletypesToInclude, DirectoriesToIgnore );
+		DirectoriesToIgnore.Add( TEXT( "/Content" ) );
+		DirectoriesToIgnore.Add( TEXT( "/DerivedDataCache" ) );
+		DirectoriesToIgnore.Add( TEXT( "/Intermediate" ) );
+		DirectoriesToIgnore.Add( TEXT( "/Documentation" ) );
+		DirectoriesToIgnore.Add( TEXT( "/NoRedist" ) );
+		DirectoriesToIgnore.Add( TEXT( "/Public" ) );
+		DirectoriesToIgnore.Add( TEXT( "/Private" ) );
+
+		// Platforms to ignore.
+		DirectoriesToIgnore.Add( TEXT( "/Android" ) );
+		DirectoriesToIgnore.Add( TEXT( "/Cygwin" ) );
+		DirectoriesToIgnore.Add( TEXT( "/HTML5" ) );
+		DirectoriesToIgnore.Add( TEXT( "/IOS" ) );
+		DirectoriesToIgnore.Add( TEXT( "/Java" ) );
+		DirectoriesToIgnore.Add( TEXT( "/Linux" ) );
+		DirectoriesToIgnore.Add( TEXT( "/Mac" ) );
+		DirectoriesToIgnore.Add( TEXT( "/PS4" ) );
+		DirectoriesToIgnore.Add( TEXT( "/XboxOne" ) );
+
+		IPlatformFile& PlatformFile = FPlatformFileManager::Get().GetPlatformFile();
+		FSymbolAndExecutableFinder SymbolAndExecutableFinder( PlatformFile, FiletypesToInclude, DirectoriesToIgnore );
 		SymbolAndExecutableFinder.Visit( *FPaths::RootDir(), true );
 
+		while( SymbolAndExecutableFinder.RecursiveDirectories.Num() )
+		{
+			const FString Directory = SymbolAndExecutableFinder.RecursiveDirectories.Pop();
+			PlatformFile.IterateDirectory( *Directory, SymbolAndExecutableFinder );
+		}
 
 		TSet<FString> SymbolPaths;
 		for( const auto& Filename : SymbolAndExecutableFinder.GetFilePaths() )
@@ -535,14 +556,18 @@ bool FWindowsPlatformStackWalkExt::GetCallstacks()
 
 	FCrashExceptionInfo& Exception = CrashInfo.Exception;
 
-	byte Context[4096] = { 0 };
+	FMemMark Mark(FMemStack::Get());
+
+	//const float int int32 FString
+	const int32 ConextSize = 4096;
+	byte* Context = new(FMemStack::Get()) byte[ConextSize];
 	ULONG DebugEvent = 0;
 	ULONG ProcessID = 0;
 	ULONG ThreadID = 0;
 	ULONG ContextSize = 0;
 
 	// Get the context of the crashed thread
-	if( Control->GetStoredEventInformation( &DebugEvent, &ProcessID, &ThreadID, Context, sizeof( Context ), &ContextSize, NULL, 0, 0) )
+	if( Control->GetStoredEventInformation( &DebugEvent, &ProcessID, &ThreadID, Context, ConextSize, &ContextSize, NULL, 0, 0) )
 	{
 		return bAtLeastOneFunctionNameFound;
 	}
@@ -561,10 +586,10 @@ bool FWindowsPlatformStackWalkExt::GetCallstacks()
 	const uint32 MaxFrames = 128;
 	const uint32 MaxFramesSize = MaxFrames * ContextSize;
 
-	DEBUG_STACK_FRAME StackFrames[MaxFrames] = { 0 };
+	DEBUG_STACK_FRAME* StackFrames = new(FMemStack::Get()) DEBUG_STACK_FRAME[MaxFrames];
 	ULONG Count = 0;
 	bool bFoundSourceFile = false;
-	void* ContextData = FMemory::Malloc( MaxFramesSize );
+	void* ContextData = FMemStack::Get().PushBytes( MaxFramesSize, 0 );
 	FMemory::Memzero( ContextData, MaxFramesSize );
 	HRESULT HR = Control->GetContextStackTrace( Context, ContextSize, StackFrames, MaxFrames, ContextData, MaxFramesSize, ContextSize, &Count );
 
@@ -633,8 +658,6 @@ bool FWindowsPlatformStackWalkExt::GetCallstacks()
 		}
 	}
 
-	FMemory::Free( ContextData );
-
 	CrashInfo.Log( FString::Printf( TEXT( " ... callstack generated!" ) ) );
 
 	return bAtLeastOneFunctionNameFound;
