CrashDebugHelper - Added support to processing crashes where builds are not placed in P4, but only in the network drive

diff --git a/Engine/Source/Developer/CrashDebugHelper/Private/CrashDebugHelper.cpp b/Engine/Source/Developer/CrashDebugHelper/Private/CrashDebugHelper.cpp
--- a/Engine/Source/Developer/CrashDebugHelper/Private/CrashDebugHelper.cpp
+++ b/Engine/Source/Developer/CrashDebugHelper/Private/CrashDebugHelper.cpp
@@ -59,7 +59,7 @@ bool ICrashDebugHelper::Init()
 		BuiltFromCL = int32(BUILT_FROM_CHANGELIST);
 	}
 
-	CrashInfo.ChangelistBuiltFrom = BuiltFromCL;
+	CrashInfo.BuiltFromCL = BuiltFromCL;
 
 	GConfig->GetString( TEXT( "Engine.CrashDebugHelper" ), TEXT( "SourceControlBuildLabelPattern" ), SourceControlBuildLabelPattern, GEngineIni );
 	GConfig->GetString( TEXT( "Engine.CrashDebugHelper" ), TEXT( "ExecutablePathPattern" ), ExecutablePathPattern, GEngineIni );
@@ -134,54 +134,101 @@ void ICrashDebugHelper::ShutdownSourceControl()
  */
 bool ICrashDebugHelper::SyncModules()
 {
-	if( CrashInfo.LabelName.Len() <= 0 )
+	// Check source control
+	if( !ISourceControlModule::Get().IsEnabled() )
 	{
-		UE_LOG( LogCrashDebugHelper, Warning, TEXT( "SyncModules: Invalid Label parameter." ) );
 		return false;
 	}
 
-	// Check source control
-	if( !ISourceControlModule::Get().IsEnabled() )
+	if( !PDBCache.UsePDBCache() )
 	{
+		UE_LOG( LogCrashDebugHelper, Warning, TEXT( "The PDB Cache is disabled, cannot proceed" ), *CrashInfo.ProductVersion );
 		return false;
 	}
 
 	const TCHAR* UESymbols = TEXT( "Rocket/Symbols/" );
-	const bool bUseNetworkPathExecutables = !CrashInfo.NetworkPath.IsEmpty();
+	const bool bUseNetworkExecutables = !CrashInfo.ExecutablesPath.IsEmpty();
+	const bool bUseNetworkSymbols = !CrashInfo.SymbolsPath.IsEmpty();
+	TArray< TSharedRef<ISourceControlLabel> > Labels = ISourceControlModule::Get().GetProvider().GetLabels( CrashInfo.LabelName );
+	
+	const FString NetworkExecutablesPath = ExecutablePathPattern.Replace( TEXT( "%PRODUCT_VERSION%" ), *CrashInfo.ProductVersion );
+	
+	/*
+		2015-01-20 Update
+		Assumes a crash from Windows
+		The PDBs for Rocket builds are now being copied to the network rather than submitted to Perforce. An example for the 4.7 branch:
+		This build is not longer submitted to the P4, so P4 label is also not created.
+
+		Build: \\epicgames.net\root\Builds\Rocket\Automated\4.7.0-2403078+++depot+UE4-Releases+4.7\Windows
+		Symbols: \\epicgames.net\root\Builds\Rocket\Automated\4.7.0-2403078+++depot+UE4-Releases+4.7\WindowsSymbols
+	*/
+	const bool bContainsProductVersion = PDBCache.UsePDBCache() && PDBCache.ContainsPDBCacheEntry( CrashInfo.ProductVersion );
+	if( bUseNetworkExecutables && bUseNetworkSymbols )
+	{
+		if( bContainsProductVersion )
+		{
+			UE_LOG( LogCrashDebugHelper, Warning, TEXT( "ProductVersion '%s' found in the PDB Cache, using it" ), *CrashInfo.ProductVersion );
+			CrashInfo.PDBCacheEntry = PDBCache.FindAndTouchPDBCacheEntry( CrashInfo.ProductVersion );
+		}
+		else
+		{
+			SCOPE_LOG_TIME_IN_SECONDS( TEXT( "SyncExecutableAndSymbolsFromNetwork" ), nullptr );
+
+			// Find all executables.
+			TArray<FString> NetworkExecutables;
+			IFileManager::Get().FindFilesRecursive( NetworkExecutables, *NetworkExecutablesPath, TEXT( "*.dll" ), true, false, false );
+			IFileManager::Get().FindFilesRecursive( NetworkExecutables, *NetworkExecutablesPath, TEXT( "*.exe" ), true, false, false );
+
+			// Find all symbols.
+			const FString NetworkSymbolsPath = NetworkExecutablesPath + TEXT("Symbols");
+			TArray<FString> NetworkSymbols;
+			IFileManager::Get().FindFilesRecursive( NetworkSymbols, *NetworkSymbolsPath, TEXT( "*.pdb" ), true, false, false );
+
+			// From=Full pathname
+			// To=Relative pathname
+			TMap<FString, FString> FilesToBeCached;
+
+			for( const auto& ExecutablePath : NetworkExecutables )
+			{
+				const FString NetworkRelativePath = ExecutablePath.Replace( *NetworkExecutablesPath, TEXT( "" ) );
+				FilesToBeCached.Add( ExecutablePath, NetworkRelativePath );
+			}
+
+			for( const auto& SymbolPath : NetworkSymbols )
+			{
+				const FString SymbolRelativePath = SymbolPath.Replace( *NetworkSymbolsPath, TEXT( "" ) );
+				FilesToBeCached.Add( SymbolPath, SymbolRelativePath );
+			}
 
+			// Initialize and add a new PDB Cache entry to the database.
+			CrashInfo.PDBCacheEntry = PDBCache.CreateAndAddPDBCacheEntryMixed( CrashInfo.ProductVersion, FilesToBeCached );
+		}
+	}
 	// Get all labels associated with the crash info's label.
 	// It should return one unique label because the crash info's label looks like this //depot/UE4-Releases/4.0/UE-CL-2047835 @see RetrieveBuildLabel
 	// We can assume that the label name is unique.
-	// 
-	TArray< TSharedRef<ISourceControlLabel> > Labels = ISourceControlModule::Get().GetProvider().GetLabels( CrashInfo.LabelName );
-	if( Labels.Num() == 1 )
+	else if( Labels.Num() == 1 )
 	{
 		TSharedRef<ISourceControlLabel> Label = Labels[0];
 		TSet<FString> FilesToSync;
 
 		// Use product version instead of label name to make a distinguish between chosen methods.
-		const bool bContainsProductVersion = PDBCache.UsePDBCache() && PDBCache.ContainsPDBCacheEntry( CrashInfo.ProductVersion );
-		const bool bContainsLabelName = PDBCache.UsePDBCache() && PDBCache.ContainsPDBCacheEntry( CrashInfo.LabelName );
+		const bool bContainsLabelName =  PDBCache.ContainsPDBCacheEntry( CrashInfo.LabelName );
 
 		if( bContainsProductVersion )
 		{
 			UE_LOG( LogCrashDebugHelper, Warning, TEXT( "ProductVersion '%s' found in the PDB Cache, using it" ), *CrashInfo.ProductVersion );
 			CrashInfo.PDBCacheEntry = PDBCache.FindAndTouchPDBCacheEntry( CrashInfo.ProductVersion );
-			return true;
 		}
 		else if( bContainsLabelName )
 		{
 			UE_LOG( LogCrashDebugHelper, Warning, TEXT( "Label '%s' found in the PDB Cache, using it" ), *CrashInfo.LabelName );
 			CrashInfo.PDBCacheEntry = PDBCache.FindAndTouchPDBCacheEntry( CrashInfo.LabelName );
-			return true;
 		}
-		else if( bUseNetworkPathExecutables )
+		else if( bUseNetworkExecutables )
 		{			
 			SCOPE_LOG_TIME_IN_SECONDS( TEXT( "SyncModulesAndNetwork" ), nullptr );
 
-			// This is a bit hacky, probably will be valid until next build system change.
-			const FString ProductNetworkPath = ExecutablePathPattern.Replace( TEXT( "%PRODUCT_VERSION%" ), *CrashInfo.ProductVersion );
-
 			// Grab information about symbols.
 			TArray< TSharedRef<class ISourceControlRevision, ESPMode::ThreadSafe> > PDBSourceControlRevisions;
 			const FString PDBsPath = FString::Printf( TEXT( "%s/%s....pdb" ), *DepotName, UESymbols );
@@ -204,8 +251,8 @@ bool ICrashDebugHelper::SyncModules()
 
 			// Find all the executables in the product network path.
 			TArray<FString> NetworkExecutables;
-			IFileManager::Get().FindFilesRecursive( NetworkExecutables, *ProductNetworkPath, TEXT( "*.dll" ), true, false, false );
-			IFileManager::Get().FindFilesRecursive( NetworkExecutables, *ProductNetworkPath, TEXT( "*.exe" ), true, false, false );
+			IFileManager::Get().FindFilesRecursive( NetworkExecutables, *NetworkExecutablesPath, TEXT( "*.dll" ), true, false, false );
+			IFileManager::Get().FindFilesRecursive( NetworkExecutables, *NetworkExecutablesPath, TEXT( "*.exe" ), true, false, false );
 
 			// From=Full pathname
 			// To=Relative pathname
@@ -220,7 +267,7 @@ bool ICrashDebugHelper::SyncModules()
 					const FString PDBFullpath = DepotRoot / PDBPath.Replace( P4_DEPOT_PREFIX, TEXT( "" ) );
 
 					const FString PDBMatch = PDBRelativePath.Replace( TEXT( "pdb" ), TEXT( "" ) );
-					const FString NetworkRelativePath = NetworkExecutableFullpath.Replace( *ProductNetworkPath, TEXT( "" ) );
+					const FString NetworkRelativePath = NetworkExecutableFullpath.Replace( *NetworkExecutablesPath, TEXT( "" ) );
 					const bool bMatch = NetworkExecutableFullpath.Contains( PDBMatch );
 					if( bMatch )
 					{
@@ -232,11 +279,8 @@ bool ICrashDebugHelper::SyncModules()
 				}
 			}
 
-			if( PDBCache.UsePDBCache() )
-			{
-				// Initialize and add a new PDB Cache entry to the database.
-				CrashInfo.PDBCacheEntry = PDBCache.CreateAndAddPDBCacheEntryMixed( CrashInfo.ProductVersion, FilesToBeCached );
-			}
+			// Initialize and add a new PDB Cache entry to the database.
+			CrashInfo.PDBCacheEntry = PDBCache.CreateAndAddPDBCacheEntryMixed( CrashInfo.ProductVersion, FilesToBeCached );
 		}
 		else
 		{
@@ -319,53 +363,32 @@ bool ICrashDebugHelper::SyncModules()
 				}
 			}
 
-			if( PDBCache.UsePDBCache() )
-			{
-				// Initialize and add a new PDB Cache entry to the database.
-				CrashInfo.PDBCacheEntry = PDBCache.CreateAndAddPDBCacheEntry( CrashInfo.LabelName, DepotRoot, DepotName, FilesToBeCached );
-			}
+			// Initialize and add a new PDB Cache entry to the database.
+			CrashInfo.PDBCacheEntry = PDBCache.CreateAndAddPDBCacheEntry( CrashInfo.LabelName, DepotRoot, DepotName, FilesToBeCached );
 		}
 	}
 	else
 	{
 		UE_LOG( LogCrashDebugHelper, Error, TEXT( "Could not find label '%s'."), *CrashInfo.LabelName );
+		return false;
 	}
 
 	return true;
 }
 
-/** 
- * Sync a single source file to the requested label
- */
 bool ICrashDebugHelper::SyncSourceFile()
 {
-	if (CrashInfo.LabelName.Len() <= 0)
-	{
-		UE_LOG( LogCrashDebugHelper, Warning, TEXT( "SyncSourceFile: Invalid Label parameter." ) );
-		return false;
-	}
-
 	// Check source control
 	if( !ISourceControlModule::Get().IsEnabled() )
 	{
 		return false;
 	}
 
-	// Sync all the dll, exes, and related symbol files
-	FString DepotPath = DepotName / CrashInfo.SourceFile;
-	TArray< TSharedRef<ISourceControlLabel> > Labels = ISourceControlModule::Get().GetProvider().GetLabels( CrashInfo.LabelName );
-	if(Labels.Num() > 0)
-	{
-		TSharedRef<ISourceControlLabel> Label = Labels[0];
-		if( Label->Sync( DepotPath ) )
-		{
-			UE_LOG( LogCrashDebugHelper, Warning, TEXT( " ... synced source file '%s'."), *DepotPath );
-		}
-	}
-	else
-	{
-		UE_LOG( LogCrashDebugHelper, Error, TEXT( "Could not find label '%s'."), *CrashInfo.LabelName );
-	}
+	// Sync a single source file to requested CL.
+	FString DepotPath = DepotName / CrashInfo.SourceFile + TEXT("@") + TTypeToString<int32>::ToString(CrashInfo.BuiltFromCL);
+	ISourceControlModule::Get().GetProvider().Execute(ISourceControlOperation::Create<FSync>(), DepotPath);
+
+	UE_LOG( LogCrashDebugHelper, Warning, TEXT( "Syncing a single source file '%s' to CL %i."), *DepotPath, CrashInfo.BuiltFromCL );
 
 	return true;
 }
@@ -430,7 +453,7 @@ void ICrashDebugHelper::AddSourceToReport()
 bool ICrashDebugHelper::AddAnnotatedSourceToReport()
 {
 	// Make sure we have a source file to interrogate
-	if( CrashInfo.SourceFile.Len() > 0 && CrashInfo.SourceLineNumber != 0 )
+	if( CrashInfo.SourceFile.Len() > 0 && CrashInfo.SourceLineNumber != 0 && !CrashInfo.LabelName.IsEmpty() )
 	{
 		// Check source control
 		if( !ISourceControlModule::Get().IsEnabled() )
@@ -498,7 +521,7 @@ const TCHAR* FCrashInfo::GetProcessorArchitecture( EProcessorArchitecture PA )
 int64 FCrashInfo::StringSize( const ANSICHAR* Line )
 {
 	int64 Size = 0;
-	if( Line != NULL )
+	if( Line != nullptr )
 	{
 		while( *Line++ != 0 )
 		{
@@ -546,7 +569,7 @@ void FCrashInfo::GenerateReport( const FString& DiagnosticsPath )
 			WriteLine( ReportFile, TCHAR_TO_UTF8( *Line ) );
 		}
 
-		Line = FString::Printf( TEXT( " ... built from changelist %d" ), ChangelistBuiltFrom );
+		Line = FString::Printf( TEXT( " ... built from changelist %d" ), BuiltFromCL );
 		WriteLine( ReportFile, TCHAR_TO_UTF8( *Line ) );
 		if( LabelName.Len() > 0 )
 		{
@@ -753,14 +776,14 @@ FString ICrashDebugHelper::RetrieveBuildLabel(int32 InChangelistNumber)
 {
 	FString FoundLabelString;
 
-	if( InChangelistNumber < 0 )
+	if( InChangelistNumber <= 0 )
 	{
 		UE_LOG(LogCrashDebugHelper, Warning, TEXT("RetrieveBuildLabel: Invalid parameters."));
 		return FoundLabelString;
 	}
 
 	// Try to find the label directly in source control by using the pattern supplied via ini
-	if ( FoundLabelString.IsEmpty() && InChangelistNumber >= 0 && !SourceControlBuildLabelPattern.IsEmpty() )
+	if ( FoundLabelString.IsEmpty() && !SourceControlBuildLabelPattern.IsEmpty() )
 	{
 		const FString ChangelistString = FString::Printf(TEXT("%d"), InChangelistNumber);
 		const FString TestLabel = SourceControlBuildLabelPattern.Replace(TEXT("%CHANGELISTNUMBER%"), *ChangelistString, ESearchCase::CaseSensitive );
@@ -782,17 +805,49 @@ FString ICrashDebugHelper::RetrieveBuildLabel(int32 InChangelistNumber)
 	return FoundLabelString;
 }
 
-void ICrashDebugHelper::RetrieveBuildLabelAndNetworkPath( int32 InChangelistNumber )
+void ICrashDebugHelper::RetrieveBuildLabelAndNetworkPaths( int32 InChangelistNumber )
 {
-	if( InChangelistNumber < 0 )
+	CrashInfo.ExecutablesPath.Empty();
+	CrashInfo.SymbolsPath.Empty();
+	CrashInfo.LabelName.Empty();
+
+	if( InChangelistNumber <= 0 )
 	{
-		UE_LOG( LogCrashDebugHelper, Warning, TEXT( "RetrieveBuildLabelAndNetworkPath: Invalid parameters." ) );
-		CrashInfo.LabelName.Empty();
-		CrashInfo.NetworkPath.Empty();
+		UE_LOG( LogCrashDebugHelper, Warning, TEXT( "RetrieveBuildLabelAndNetworkPaths: Invalid parameters." ) );
+	}
+
+	// Try to find the network path by using the pattern supplied via ini.
+	// If this step successes, we will grab the executable from the network path instead of P4.
+	// At this moment we only care about UE4 releases.
+	const bool bCanUseNetworkPath = DepotName.Contains( TEXT( "UE4-Releases" ) );
+	bool bFound = false;
+	if( bCanUseNetworkPath && !ExecutablePathPattern.IsEmpty() )
+	{
+		const FString TestNetworkExecutablesPath = ExecutablePathPattern.Replace( TEXT( "%PRODUCT_VERSION%" ), *CrashInfo.ProductVersion );
+		const FString TestNetworkSymbolsPath = TestNetworkExecutablesPath + TEXT("Symbols");
+		const bool bHasExecutables = IFileManager::Get().DirectoryExists( *TestNetworkExecutablesPath );
+		const bool bHasSymbols = IFileManager::Get().DirectoryExists( *TestNetworkExecutablesPath );
+
+		if( bHasExecutables && bHasSymbols )
+		{
+			CrashInfo.ExecutablesPath = TestNetworkExecutablesPath;
+			CrashInfo.SymbolsPath = TestNetworkSymbolsPath;
+			bFound = true;
+			UE_LOG( LogCrashDebugHelper, Log, TEXT( "RetrieveBuildLabelAndNetworkPaths: Found network path for executables and symbols '%s'." ), *CrashInfo.ExecutablesPath );
+		}
+		else if( bHasExecutables )
+		{
+			CrashInfo.ExecutablesPath = TestNetworkExecutablesPath;
+			UE_LOG( LogCrashDebugHelper, Log, TEXT( "RetrieveBuildLabelAndNetworkPaths: Found network path for executables '%s'." ), *CrashInfo.ExecutablesPath );
+		}
+		else
+		{
+			UE_LOG( LogCrashDebugHelper, Log, TEXT( "RetrieveBuildLabelAndNetworkPaths: Network path for executables '%s' not found." ), *TestNetworkExecutablesPath );
+		}
 	}
 
 	// Try to find the label directly in source control by using the pattern supplied via ini.
-	if( InChangelistNumber >= 0 && !SourceControlBuildLabelPattern.IsEmpty() )
+	if( !bFound && !SourceControlBuildLabelPattern.IsEmpty() )
 	{
 		const FString ChangelistString = FString::Printf( TEXT( "%d" ), InChangelistNumber );
 		const FString TestLabelWithCL = SourceControlBuildLabelPattern.Replace( TEXT( "%CHANGELISTNUMBER%" ), *ChangelistString, ESearchCase::CaseSensitive );
@@ -803,41 +858,12 @@ void ICrashDebugHelper::RetrieveBuildLabelAndNetworkPath( int32 InChangelistNumb
 			// If we found more than one label, warn about it and just use the first one
 			if( Labels.Num() > 1 )
 			{
-				UE_LOG( LogCrashDebugHelper, Warning, TEXT( "RetrieveBuildLabelAndNetworkPath: More than one build label found with pattern '%s' - Using label '%s'" ), *TestLabelWithCL, *Labels[LabelIndex]->GetName() );
+				UE_LOG( LogCrashDebugHelper, Warning, TEXT( "RetrieveBuildLabelAndNetworkPaths: More than one build label found with pattern '%s' - Using label '%s'" ), *TestLabelWithCL, *Labels[LabelIndex]->GetName() );
 			}
 
 			CrashInfo.LabelName = Labels[LabelIndex]->GetName();
-			UE_LOG( LogCrashDebugHelper, Log, TEXT( "RetrieveBuildLabelAndNetworkPath: Found label '%s' matching pattern '%s' in source control." ), *CrashInfo.LabelName, *TestLabelWithCL );
-		}
-	}
-	else
-	{
-		CrashInfo.LabelName.Empty();
-	}
-
-	// Try to find the network path by using the pattern supplied via ini.
-	// If this step successes, we will grab the executable from the network path instead of P4.
-	// At this moment we only care about UE4 releases.
-	const bool bCanUseNetworkPath = DepotName.Contains( TEXT( "UE4-Releases" ) );
-	if( bCanUseNetworkPath && InChangelistNumber >= 0 && !ExecutablePathPattern.IsEmpty() )
-	{
-		const FString TestNetworkPath = ExecutablePathPattern.Replace( TEXT( "%PRODUCT_VERSION%" ), *CrashInfo.ProductVersion );
-		const bool bIsValidPath = IFileManager::Get().DirectoryExists( *TestNetworkPath );
-
-		if( bIsValidPath )
-		{
-			CrashInfo.NetworkPath = TestNetworkPath;
-			UE_LOG( LogCrashDebugHelper, Log, TEXT( "RetrieveBuildLabelAndNetworkPath: Found network path '%s'." ), *CrashInfo.NetworkPath );
+			UE_LOG( LogCrashDebugHelper, Log, TEXT( "RetrieveBuildLabelAndNetworkPaths: Found label '%s' matching pattern '%s' in source control." ), *CrashInfo.LabelName, *TestLabelWithCL );
 		}
-		else
-		{
-			CrashInfo.NetworkPath.Empty();
-			UE_LOG( LogCrashDebugHelper, Log, TEXT( "RetrieveBuildLabelAndNetworkPath: Network path '%s' not found." ), *TestNetworkPath );
-		}
-	}
-	else
-	{
-		CrashInfo.NetworkPath.Empty();
 	}
 }
 
diff --git a/Engine/Source/Developer/CrashDebugHelper/Private/Mac/CrashDebugHelperMac.cpp b/Engine/Source/Developer/CrashDebugHelper/Private/Mac/CrashDebugHelperMac.cpp
--- a/Engine/Source/Developer/CrashDebugHelper/Private/Mac/CrashDebugHelperMac.cpp
+++ b/Engine/Source/Developer/CrashDebugHelper/Private/Mac/CrashDebugHelperMac.cpp
@@ -536,7 +536,7 @@ bool FCrashDebugHelperMac::CreateMinidumpDiagnosticReport( const FString& InCras
 			
 			if(Result >= 4)
 			{
-				CrashInfo.ChangelistBuiltFrom = CLNumber;
+				CrashInfo.BuiltFromCL = CLNumber;
 			}
 			
 			if(Result == 5 && Branch.Len() > 0)
diff --git a/Engine/Source/Developer/CrashDebugHelper/Private/Windows/MinidumpParser.cpp b/Engine/Source/Developer/CrashDebugHelper/Private/Windows/MinidumpParser.cpp
--- a/Engine/Source/Developer/CrashDebugHelper/Private/Windows/MinidumpParser.cpp
+++ b/Engine/Source/Developer/CrashDebugHelper/Private/Windows/MinidumpParser.cpp
@@ -26,7 +26,7 @@ bool FCrashDebugHelperWindows::CreateMinidumpDiagnosticReport( const FString& In
 	bool bAtLeastOneFunctionNameFoundInCallstack = false;
 	if( bReady && WindowsStackWalkExt.OpenDumpFile( InCrashDumpFilename ) )
 	{
-		if( CrashInfo.ChangelistBuiltFrom != FCrashInfo::INVALID_CHANGELIST )
+		if( CrashInfo.BuiltFromCL != FCrashInfo::INVALID_CHANGELIST )
 		{
 			// Get the build version from the crash info.
 			FCrashModuleInfo ExeFileVersion;
@@ -34,14 +34,19 @@ bool FCrashDebugHelperWindows::CreateMinidumpDiagnosticReport( const FString& In
 
 			// @see AutomationTool.CommandUtils.EscapePath 
 			const FString CleanedDepotName = CrashInfo.DepotName.Replace( TEXT( ":" ), TEXT( "" ) ).Replace( TEXT( "/" ), TEXT( "+" ) ).Replace( TEXT( "\\" ), TEXT( "+" ) ).Replace( TEXT( " " ), TEXT( "+" ) );
-			const FEngineVersion EngineVersion( ExeFileVersion.Major, ExeFileVersion.Minor, ExeFileVersion.Patch, CrashInfo.ChangelistBuiltFrom, CleanedDepotName );
+			const FEngineVersion EngineVersion( ExeFileVersion.Major, ExeFileVersion.Minor, ExeFileVersion.Patch, CrashInfo.BuiltFromCL, CleanedDepotName );
 			CrashInfo.ProductVersion = EngineVersion.ToString();
 
 			// CrashInfo now contains a changelist to lookup a label for
 			if( bSyncSymbols )
 			{
-				RetrieveBuildLabelAndNetworkPath( CrashInfo.ChangelistBuiltFrom );
-				SyncModules();
+				RetrieveBuildLabelAndNetworkPaths( CrashInfo.BuiltFromCL );
+				const bool bSynced = SyncModules();
+				// Without symbols we can't decode the provided minidump.
+				if( !bSynced )
+				{
+					return 0;
+				}
 			}
 
 			// Initialise the symbol options
@@ -68,9 +73,9 @@ bool FCrashDebugHelperWindows::CreateMinidumpDiagnosticReport( const FString& In
 			// Sync the source file where the crash occurred
 			if( CrashInfo.SourceFile.Len() > 0 )
 			{
-				if( bSyncSymbols )
+				if( bSyncSymbols && CrashInfo.BuiltFromCL > 0 )
 				{
-					CrashInfo.Log( FString::Printf( TEXT( " ... using label '%s' to sync crash source file" ), *CrashInfo.LabelName ) );
+					CrashInfo.Log( FString::Printf( TEXT( " ... using CL %i to sync crash source file" ), CrashInfo.BuiltFromCL ) );
 					SyncSourceFile();
 				}
 
diff --git a/Engine/Source/Developer/CrashDebugHelper/Public/CrashDebugHelper.h b/Engine/Source/Developer/CrashDebugHelper/Public/CrashDebugHelper.h
--- a/Engine/Source/Developer/CrashDebugHelper/Public/CrashDebugHelper.h
+++ b/Engine/Source/Developer/CrashDebugHelper/Public/CrashDebugHelper.h
@@ -140,13 +140,16 @@ public:
 	FString ProductVersion;
 
 	/** CL built from. */
-	int32 ChangelistBuiltFrom;
+	int32 BuiltFromCL;
 
 	/** The label the describes the executables and symbols. */
 	FString LabelName;
 
 	/** The network path where the executables are stored. */
-	FString NetworkPath;
+	FString ExecutablesPath;
+
+	/** The network path where the symbols are stored. */
+	FString SymbolsPath;
 
 	FString SourceFile;
 	uint64 SourceLineNumber;
@@ -164,7 +167,7 @@ public:
 	FPDBCacheEntryPtr PDBCacheEntry;
 
 	FCrashInfo()
-		: ChangelistBuiltFrom( INVALID_CHANGELIST )
+		: BuiltFromCL( INVALID_CHANGELIST )
 		, SourceLineNumber( 0 )
 	{
 	}
@@ -560,9 +563,7 @@ public:
 	virtual bool SyncModules();
 
 	/**
-	 *	Sync a single source file to the requested label.
-	 *
-	 *	@return	bool		true if successful, false if not
+	 *	Sync a single source file to the requested CL.
 	 */
 	virtual bool SyncSourceFile();
 
@@ -587,13 +588,11 @@ protected:
 
 	/**
 	 *	Retrieve the build label and the network path for the given changelist number.
-	 *
-	 *	@param	InChangelistNumber	The changelist number to retrieve the label for
 	 */
-	void RetrieveBuildLabelAndNetworkPath( int32 InChangelistNumber );
+	void RetrieveBuildLabelAndNetworkPaths( int32 InChangelistNumber );
 
-	bool ReadSourceFile( const TCHAR* InFilename, TArray<FString>& OutStrings );
 
+	bool ReadSourceFile( const TCHAR* InFilename, TArray<FString>& OutStrings );
 public:
 	bool InitSourceControl(bool bShowLogin);
 	void ShutdownSourceControl();
