CrashDebugHelper - Removed all obsolete functionality, will reimplement later, cleaned a lof of stuff, better support for network builds

diff --git a/Engine/Source/Developer/CrashDebugHelper/CrashDebugHelper.Build.cs b/Engine/Source/Developer/CrashDebugHelper/CrashDebugHelper.Build.cs
--- a/Engine/Source/Developer/CrashDebugHelper/CrashDebugHelper.Build.cs
+++ b/Engine/Source/Developer/CrashDebugHelper/CrashDebugHelper.Build.cs
@@ -4,16 +4,22 @@ using UnrealBuildTool;
 
 public class CrashDebugHelper : ModuleRules
 {
-	public CrashDebugHelper(TargetInfo Target)
+	public CrashDebugHelper( TargetInfo Target )
 	{
-		PrivateIncludePaths.Add("Developer/CrashDebugHelper/Private");
-		PrivateIncludePaths.Add("ThirdParty/PLCrashReporter/plcrashreporter-master-5ae3b0a/Source");
+		PrivateIncludePaths.AddRange(
+		new string[] {
+				"Developer/CrashDebugHelper/Private/",
+				"Developer/CrashDebugHelper/Private/Linux",
+				"Developer/CrashDebugHelper/Private/Mac",
+				"Developer/CrashDebugHelper/Private/Windows",
+			}
+		);
+		PrivateIncludePaths.Add( "ThirdParty/PLCrashReporter/plcrashreporter-master-5ae3b0a/Source" );
 
-		PrivateDependencyModuleNames.AddRange(
+		PublicDependencyModuleNames.AddRange(
 			new string[] {
 				"Core",
-				"SourceControl",
-				"DatabaseSupport"
+				"SourceControl"
 			}
 		);
 	}
diff --git a/Engine/Source/Developer/CrashDebugHelper/Private/CrashDebugHelper.cpp b/Engine/Source/Developer/CrashDebugHelper/Private/CrashDebugHelper.cpp
--- a/Engine/Source/Developer/CrashDebugHelper/Private/CrashDebugHelper.cpp
+++ b/Engine/Source/Developer/CrashDebugHelper/Private/CrashDebugHelper.cpp
@@ -1,8 +1,9 @@
 // Copyright 1998-2015 Epic Games, Inc. All Rights Reserved.
 
 #include "CrashDebugHelperPrivatePCH.h"
-#include "CrashDebugHelper.h"
-#include "Database.h"
+#include "CrashDebugPDBCache.h"
+
+#include "EngineVersion.h"
 #include "ISourceControlModule.h"
 #include "ISourceControlLabel.h"
 #include "ISourceControlRevision.h"
@@ -13,72 +14,84 @@
 	#define MINIDUMPDIAGNOSTICS	0
 #endif
 
-static const TCHAR* P4_DEPOT_PREFIX = TEXT( "//depot/" );
+const TCHAR* ICrashDebugHelper::P4_DEPOT_PREFIX = TEXT( "//depot/" );
 
-/** 
- * Global initialisation of this module
- */
 bool ICrashDebugHelper::Init()
 {
 	bInitialized = true;
 
-	// Look up the depot name
-	// Try to use the command line param
-	FString CmdLineBranchName;
-	if ( FParse::Value(FCommandLine::Get(), TEXT("BranchName="), CmdLineBranchName) )
-	{
-		DepotName = FString::Printf( TEXT( "%s%s" ), P4_DEPOT_PREFIX, *CmdLineBranchName );
-	}
-	// Try to use what is configured in ini
-	else if (GConfig->GetString(TEXT("Engine.CrashDebugHelper"), TEXT("DepotName"), DepotName, GEngineIni) == true)
+	// Check if we have a valid EngineVersion, if so use it.
+	FString CmdEngineVersion;
+	const bool bHasEngineVersion = FParse::Value( FCommandLine::Get(), TEXT( "EngineVersion=" ), CmdEngineVersion );
+	if( bHasEngineVersion )
 	{
-		// Hack to get around ini files treating '//' as an inlined comment
-		DepotName.ReplaceInline(TEXT("\\"), TEXT("/"));
+		FEngineVersion EngineVersion;
+		FEngineVersion::Parse( CmdEngineVersion, EngineVersion );
+
+		// Clean branch name.
+		CrashInfo.DepotName = EngineVersion.GetBranch();
+		CrashInfo.BuiltFromCL = (int32)EngineVersion.GetChangelist();
+
+		CrashInfo.EngineVersion = CmdEngineVersion;
 	}
-	// Default to BRANCH_NAME
 	else
 	{
-		DepotName = FString::Printf( TEXT( "%s%s" ), P4_DEPOT_PREFIX, TEXT( BRANCH_NAME ) );
-	}
+		// Look up the depot name
+		// Try to use the command line param
+		FString DepotName;
+		FString CmdLineBranchName;
+		if( FParse::Value( FCommandLine::Get(), TEXT( "BranchName=" ), CmdLineBranchName ) )
+		{
+			DepotName = FString::Printf( TEXT( "%s%s" ), P4_DEPOT_PREFIX, *CmdLineBranchName );
+		}
+		// Default to BRANCH_NAME
+		else
+		{
+			DepotName = FString::Printf( TEXT( "%s%s" ), P4_DEPOT_PREFIX, TEXT( BRANCH_NAME ) );
+		}
 
-	CrashInfo.DepotName = DepotName;
+		CrashInfo.DepotName = DepotName;
 
-	// Try to get the BuiltFromCL from command line to use this instead of attempting to locate the CL in the minidump
-	FString CmdLineBuiltFromCL;
-	int32 BuiltFromCL = -1;
-	if (FParse::Value(FCommandLine::Get(), TEXT("BuiltFromCL="), CmdLineBuiltFromCL))
-	{
-		if ( !CmdLineBuiltFromCL.IsEmpty() )
+		// Try to get the BuiltFromCL from command line to use this instead of attempting to locate the CL in the minidump
+		FString CmdLineBuiltFromCL;
+		int32 BuiltFromCL = -1;
+		if( FParse::Value( FCommandLine::Get(), TEXT( "BuiltFromCL=" ), CmdLineBuiltFromCL ) )
+		{
+			if( !CmdLineBuiltFromCL.IsEmpty() )
+			{
+				BuiltFromCL = FCString::Atoi( *CmdLineBuiltFromCL );
+			}
+		}
+		// Default to BUILT_FROM_CHANGELIST.
+		else
 		{
-			BuiltFromCL = FCString::Atoi(*CmdLineBuiltFromCL);
+			BuiltFromCL = int32( BUILT_FROM_CHANGELIST );
 		}
-	}
-	// Default to BUILT_FROM_CHANGELIST.
-	else
-	{
-		BuiltFromCL = int32(BUILT_FROM_CHANGELIST);
-	}
-
-	CrashInfo.BuiltFromCL = BuiltFromCL;
 
+		CrashInfo.BuiltFromCL = BuiltFromCL;
+	}
+	
 	GConfig->GetString( TEXT( "Engine.CrashDebugHelper" ), TEXT( "SourceControlBuildLabelPattern" ), SourceControlBuildLabelPattern, GEngineIni );
-	GConfig->GetString( TEXT( "Engine.CrashDebugHelper" ), TEXT( "ExecutablePathPattern" ), ExecutablePathPattern, GEngineIni );
 
-	// Look up the local symbol store - fail if not found
-	if (GConfig->GetString(TEXT("Engine.CrashDebugHelper"), TEXT("LocalSymbolStore"), LocalSymbolStore, GEngineIni) == false)
-	{
-		UE_LOG(LogCrashDebugHelper, Warning, TEXT("Failed to get LocalSymbolStore from ini file... crash handling disabled"));
-		bInitialized = false;
-	}
+	GConfig->GetArray( TEXT( "Engine.CrashDebugHelper" ), TEXT( "ExecutablePathPattern" ), ExecutablePathPatterns, GEngineIni );
+	GConfig->GetArray( TEXT( "Engine.CrashDebugHelper" ), TEXT( "SymbolPathPattern" ), SymbolPathPatterns, GEngineIni );
+	GConfig->GetArray( TEXT( "Engine.CrashDebugHelper" ), TEXT( "Branch" ), Branches, GEngineIni );
+	const bool bCanUseSearchPatterns = Branches.Num() == ExecutablePathPatterns.Num() && ExecutablePathPatterns.Num() == SymbolPathPatterns.Num() && Branches.Num() > 0;
+	UE_CLOG( !bCanUseSearchPatterns, LogCrashDebugHelper, Warning, TEXT( "Search patterns don't match" ) );
 
-	if( !GConfig->GetString( TEXT( "Engine.CrashDebugHelper" ), TEXT( "DepotRoot" ), DepotRoot, GEngineIni ) )
+	GConfig->GetString( TEXT( "Engine.CrashDebugHelper" ), TEXT( "DepotRoot" ), DepotRoot, GEngineIni );
+	const bool bHasDepotRoot = IFileManager::Get().DirectoryExists( *DepotRoot );
+	UE_CLOG( !bHasDepotRoot, LogCrashDebugHelper, Warning, TEXT( "DepotRoot: %s is not valid" ), *DepotRoot );
+
+	if( bCanUseSearchPatterns && bHasDepotRoot )
 	{
-		UE_LOG( LogCrashDebugHelper, Warning, TEXT( "Failed to get DepotRoot from ini file... PDB Cache disabled" ) );
+		FPDBCache::Get().Init();
 	}
 	else
 	{
-		PDBCache.Init();
+		UE_LOG( LogCrashDebugHelper, Warning, TEXT( "PDB Cache disabled" ) );
 	}
+	
 
 	return bInitialized;
 }
@@ -103,6 +116,7 @@ bool ICrashDebugHelper::InitSourceControl(bool bShowLogin)
 
 		// Attempt to load in a source control module
 		ISourceControlModule::Get().GetProvider().Init();
+
 #if !MINIDUMPDIAGNOSTICS
 		if ((ISourceControlModule::Get().GetProvider().IsAvailable() == false) || bShowLogin)
 		{
@@ -129,9 +143,7 @@ void ICrashDebugHelper::ShutdownSourceControl()
 	ISourceControlModule::Get().GetProvider().Close();
 }
 
-/** 
- * Sync the branch root relative file names to the requested label
- */
+
 bool ICrashDebugHelper::SyncModules()
 {
 	// Check source control
@@ -140,35 +152,25 @@ bool ICrashDebugHelper::SyncModules()
 		return false;
 	}
 
-	if( !PDBCache.UsePDBCache() )
+	if( !FPDBCache::Get().UsePDBCache() )
 	{
-		UE_LOG( LogCrashDebugHelper, Warning, TEXT( "The PDB Cache is disabled, cannot proceed" ), *CrashInfo.ProductVersion );
+		UE_LOG( LogCrashDebugHelper, Warning, TEXT( "The PDB Cache is disabled, cannot proceed, %s" ), *CrashInfo.EngineVersion );
 		return false;
 	}
 
+	// @TODO yrx 2015-02-23 Obsolete, remove after 4.8
 	const TCHAR* UESymbols = TEXT( "Rocket/Symbols/" );
-	const bool bUseNetworkExecutables = !CrashInfo.ExecutablesPath.IsEmpty();
-	const bool bUseNetworkSymbols = !CrashInfo.SymbolsPath.IsEmpty();
+	const bool bHasExecutable = !CrashInfo.ExecutablesPath.IsEmpty();
+	const bool bHasSymbols = !CrashInfo.SymbolsPath.IsEmpty();
 	TArray< TSharedRef<ISourceControlLabel> > Labels = ISourceControlModule::Get().GetProvider().GetLabels( CrashInfo.LabelName );
 	
-	const FString NetworkExecutablesPath = ExecutablePathPattern.Replace( TEXT( "%PRODUCT_VERSION%" ), *CrashInfo.ProductVersion );
-	
-	/*
-		2015-01-20 Update
-		Assumes a crash from Windows
-		The PDBs for Rocket builds are now being copied to the network rather than submitted to Perforce. An example for the 4.7 branch:
-		This build is not longer submitted to the P4, so P4 label is also not created.
-
-		Build: \\epicgames.net\root\Builds\Rocket\Automated\4.7.0-2403078+++depot+UE4-Releases+4.7\Windows
-		Symbols: \\epicgames.net\root\Builds\Rocket\Automated\4.7.0-2403078+++depot+UE4-Releases+4.7\WindowsSymbols
-	*/
-	const bool bContainsProductVersion = PDBCache.UsePDBCache() && PDBCache.ContainsPDBCacheEntry( CrashInfo.ProductVersion );
-	if( bUseNetworkExecutables && bUseNetworkSymbols )
+	const bool bContainsProductVersion = FPDBCache::Get().ContainsPDBCacheEntry( CrashInfo.EngineVersion );
+	if( bHasExecutable && bHasSymbols )
 	{
 		if( bContainsProductVersion )
 		{
-			UE_LOG( LogCrashDebugHelper, Warning, TEXT( "ProductVersion '%s' found in the PDB Cache, using it" ), *CrashInfo.ProductVersion );
-			CrashInfo.PDBCacheEntry = PDBCache.FindAndTouchPDBCacheEntry( CrashInfo.ProductVersion );
+			UE_LOG( LogCrashDebugHelper, Warning, TEXT( "Using cached storage: %s" ), *CrashInfo.EngineVersion );
+			CrashInfo.PDBCacheEntry = FPDBCache::Get().FindAndTouchPDBCacheEntry( CrashInfo.EngineVersion );
 		}
 		else
 		{
@@ -176,13 +178,12 @@ bool ICrashDebugHelper::SyncModules()
 
 			// Find all executables.
 			TArray<FString> NetworkExecutables;
-			IFileManager::Get().FindFilesRecursive( NetworkExecutables, *NetworkExecutablesPath, TEXT( "*.dll" ), true, false, false );
-			IFileManager::Get().FindFilesRecursive( NetworkExecutables, *NetworkExecutablesPath, TEXT( "*.exe" ), true, false, false );
+			IFileManager::Get().FindFilesRecursive( NetworkExecutables, *CrashInfo.ExecutablesPath, TEXT( "*.dll" ), true, false, false );
+			IFileManager::Get().FindFilesRecursive( NetworkExecutables, *CrashInfo.ExecutablesPath, TEXT( "*.exe" ), true, false, false );
 
 			// Find all symbols.
-			const FString NetworkSymbolsPath = NetworkExecutablesPath + TEXT("Symbols");
 			TArray<FString> NetworkSymbols;
-			IFileManager::Get().FindFilesRecursive( NetworkSymbols, *NetworkSymbolsPath, TEXT( "*.pdb" ), true, false, false );
+			IFileManager::Get().FindFilesRecursive( NetworkSymbols, *CrashInfo.SymbolsPath, TEXT( "*.pdb" ), true, false, false );
 
 			// From=Full pathname
 			// To=Relative pathname
@@ -190,48 +191,46 @@ bool ICrashDebugHelper::SyncModules()
 
 			for( const auto& ExecutablePath : NetworkExecutables )
 			{
-				const FString NetworkRelativePath = ExecutablePath.Replace( *NetworkExecutablesPath, TEXT( "" ) );
+				const FString NetworkRelativePath = ExecutablePath.Replace( *CrashInfo.ExecutablesPath, TEXT( "" ) );
 				FilesToBeCached.Add( ExecutablePath, NetworkRelativePath );
 			}
 
 			for( const auto& SymbolPath : NetworkSymbols )
 			{
-				const FString SymbolRelativePath = SymbolPath.Replace( *NetworkSymbolsPath, TEXT( "" ) );
+				const FString SymbolRelativePath = SymbolPath.Replace( *CrashInfo.SymbolsPath, TEXT( "" ) );
 				FilesToBeCached.Add( SymbolPath, SymbolRelativePath );
 			}
 
 			// Initialize and add a new PDB Cache entry to the database.
-			CrashInfo.PDBCacheEntry = PDBCache.CreateAndAddPDBCacheEntryMixed( CrashInfo.ProductVersion, FilesToBeCached );
+			CrashInfo.PDBCacheEntry = FPDBCache::Get().CreateAndAddPDBCacheEntryMixed( CrashInfo.EngineVersion, FilesToBeCached );
 		}
 	}
 	// Get all labels associated with the crash info's label.
-	// It should return one unique label because the crash info's label looks like this //depot/UE4-Releases/4.0/UE-CL-2047835 @see RetrieveBuildLabel
-	// We can assume that the label name is unique.
-	else if( Labels.Num() == 1 )
+	else if( Labels.Num() >= 1 )
 	{
 		TSharedRef<ISourceControlLabel> Label = Labels[0];
 		TSet<FString> FilesToSync;
 
 		// Use product version instead of label name to make a distinguish between chosen methods.
-		const bool bContainsLabelName =  PDBCache.ContainsPDBCacheEntry( CrashInfo.LabelName );
+		const bool bContainsLabelName = FPDBCache::Get().ContainsPDBCacheEntry( CrashInfo.LabelName );
 
 		if( bContainsProductVersion )
 		{
-			UE_LOG( LogCrashDebugHelper, Warning, TEXT( "ProductVersion '%s' found in the PDB Cache, using it" ), *CrashInfo.ProductVersion );
-			CrashInfo.PDBCacheEntry = PDBCache.FindAndTouchPDBCacheEntry( CrashInfo.ProductVersion );
+			UE_LOG( LogCrashDebugHelper, Warning, TEXT( "Using cached storage: %s" ), *CrashInfo.EngineVersion );
+			CrashInfo.PDBCacheEntry = FPDBCache::Get().FindAndTouchPDBCacheEntry( CrashInfo.EngineVersion );
 		}
 		else if( bContainsLabelName )
 		{
-			UE_LOG( LogCrashDebugHelper, Warning, TEXT( "Label '%s' found in the PDB Cache, using it" ), *CrashInfo.LabelName );
-			CrashInfo.PDBCacheEntry = PDBCache.FindAndTouchPDBCacheEntry( CrashInfo.LabelName );
+			UE_LOG( LogCrashDebugHelper, Warning, TEXT( "Using cached storage: %s" ), *CrashInfo.LabelName );
+			CrashInfo.PDBCacheEntry = FPDBCache::Get().FindAndTouchPDBCacheEntry( CrashInfo.LabelName );
 		}
-		else if( bUseNetworkExecutables )
+		else if( bHasExecutable )
 		{			
 			SCOPE_LOG_TIME_IN_SECONDS( TEXT( "SyncModulesAndNetwork" ), nullptr );
 
 			// Grab information about symbols.
 			TArray< TSharedRef<class ISourceControlRevision, ESPMode::ThreadSafe> > PDBSourceControlRevisions;
-			const FString PDBsPath = FString::Printf( TEXT( "%s/%s....pdb" ), *DepotName, UESymbols );
+			const FString PDBsPath = FString::Printf( TEXT( "%s/%s....pdb" ), *CrashInfo.DepotName, UESymbols );
 			Label->GetFileRevisions( PDBsPath, PDBSourceControlRevisions );
 
 			TSet<FString> PDBPaths;
@@ -245,14 +244,14 @@ bool ICrashDebugHelper::SyncModules()
 			{
 				if( Label->Sync( PDBPath ) )
 				{
-					UE_LOG( LogCrashDebugHelper, Warning, TEXT( "Synced PDB '%s'" ), *PDBPath );
+					UE_LOG( LogCrashDebugHelper, Warning, TEXT( "Synced PDB: %s" ), *PDBPath );
 				}
 			}
 
 			// Find all the executables in the product network path.
 			TArray<FString> NetworkExecutables;
-			IFileManager::Get().FindFilesRecursive( NetworkExecutables, *NetworkExecutablesPath, TEXT( "*.dll" ), true, false, false );
-			IFileManager::Get().FindFilesRecursive( NetworkExecutables, *NetworkExecutablesPath, TEXT( "*.exe" ), true, false, false );
+			IFileManager::Get().FindFilesRecursive( NetworkExecutables, *CrashInfo.ExecutablesPath, TEXT( "*.dll" ), true, false, false );
+			IFileManager::Get().FindFilesRecursive( NetworkExecutables, *CrashInfo.ExecutablesPath, TEXT( "*.exe" ), true, false, false );
 
 			// From=Full pathname
 			// To=Relative pathname
@@ -263,15 +262,15 @@ bool ICrashDebugHelper::SyncModules()
 			{
 				for( const auto& PDBPath : PDBPaths )
 				{
-					const FString PDBRelativePath = PDBPath.Replace( *DepotName, TEXT( "" ) ).Replace( UESymbols, TEXT( "" ) );
+					const FString PDBRelativePath = PDBPath.Replace( *CrashInfo.DepotName, TEXT( "" ) ).Replace( UESymbols, TEXT( "" ) );
 					const FString PDBFullpath = DepotRoot / PDBPath.Replace( P4_DEPOT_PREFIX, TEXT( "" ) );
 
 					const FString PDBMatch = PDBRelativePath.Replace( TEXT( "pdb" ), TEXT( "" ) );
-					const FString NetworkRelativePath = NetworkExecutableFullpath.Replace( *NetworkExecutablesPath, TEXT( "" ) );
+					const FString NetworkRelativePath = NetworkExecutableFullpath.Replace( *CrashInfo.ExecutablesPath, TEXT( "" ) );
 					const bool bMatch = NetworkExecutableFullpath.Contains( PDBMatch );
 					if( bMatch )
 					{
-						// From -> Where ?
+						// From -> Where
 						FilesToBeCached.Add( NetworkExecutableFullpath, NetworkRelativePath );
 						FilesToBeCached.Add( PDBFullpath, PDBRelativePath );
 						break;
@@ -280,7 +279,7 @@ bool ICrashDebugHelper::SyncModules()
 			}
 
 			// Initialize and add a new PDB Cache entry to the database.
-			CrashInfo.PDBCacheEntry = PDBCache.CreateAndAddPDBCacheEntryMixed( CrashInfo.ProductVersion, FilesToBeCached );
+			CrashInfo.PDBCacheEntry = FPDBCache::Get().CreateAndAddPDBCacheEntryMixed( CrashInfo.EngineVersion, FilesToBeCached );
 		}
 		else
 		{
@@ -289,42 +288,42 @@ bool ICrashDebugHelper::SyncModules()
 			//@TODO: MAC: Excluding labels for Mac since we are only syncing windows binaries here...
 			if( Label->GetName().Contains( TEXT( "Mac" ) ) )
 			{
-				UE_LOG( LogCrashDebugHelper, Log, TEXT( " Skipping Mac label '%s' when syncing modules." ), *Label->GetName() );
+				UE_LOG( LogCrashDebugHelper, Log, TEXT( "Skipping Mac label: %s" ), *Label->GetName() );
 			}
 			else
 			{
 				// Sync all the dll, exes, and related symbol files
-				UE_LOG( LogCrashDebugHelper, Log, TEXT( " Syncing modules with label '%s'." ), *Label->GetName() );
+				UE_LOG( LogCrashDebugHelper, Log, TEXT( "Syncing modules with label: %s" ), *Label->GetName() );
 
 				SCOPE_LOG_TIME_IN_SECONDS( TEXT( "SyncModules" ), nullptr );
 
 				// Grab all dll and pdb files for the specified label.
 				TArray< TSharedRef<class ISourceControlRevision, ESPMode::ThreadSafe> > DLLSourceControlRevisions;
-				const FString DLLsPath = FString::Printf( TEXT( "%s/....dll" ), *DepotName );
+				const FString DLLsPath = FString::Printf( TEXT( "%s/....dll" ), *CrashInfo.DepotName );
 				Label->GetFileRevisions( DLLsPath, DLLSourceControlRevisions );
 
 				TArray< TSharedRef<class ISourceControlRevision, ESPMode::ThreadSafe> > EXESourceControlRevisions;
-				const FString EXEsPath = FString::Printf( TEXT( "%s/....exe" ), *DepotName );
+				const FString EXEsPath = FString::Printf( TEXT( "%s/....exe" ), *CrashInfo.DepotName );
 				Label->GetFileRevisions( EXEsPath, EXESourceControlRevisions );
 
 				TArray< TSharedRef<class ISourceControlRevision, ESPMode::ThreadSafe> > PDBSourceControlRevisions;
-				const FString PDBsPath = FString::Printf( TEXT( "%s/....pdb" ), *DepotName );
+				const FString PDBsPath = FString::Printf( TEXT( "%s/....pdb" ), *CrashInfo.DepotName );
 				Label->GetFileRevisions( PDBsPath, PDBSourceControlRevisions );
 
 				TSet<FString> ModulesPaths;
 				for( const auto& DLLSrc : DLLSourceControlRevisions )
 				{
-					ModulesPaths.Add( DLLSrc->GetFilename().Replace( *DepotName, TEXT( "" ) ) );
+					ModulesPaths.Add( DLLSrc->GetFilename().Replace( *CrashInfo.DepotName, TEXT( "" ) ) );
 				}
 				for( const auto& EXESrc : EXESourceControlRevisions )
 				{
-					ModulesPaths.Add( EXESrc->GetFilename().Replace( *DepotName, TEXT( "" ) ) );
+					ModulesPaths.Add( EXESrc->GetFilename().Replace( *CrashInfo.DepotName, TEXT( "" ) ) );
 				}
 
 				TSet<FString> PDBPaths;
 				for( const auto& PDBSrc : PDBSourceControlRevisions )
 				{
-					PDBPaths.Add( PDBSrc->GetFilename().Replace( *DepotName, TEXT( "" ) ) );
+					PDBPaths.Add( PDBSrc->GetFilename().Replace( *CrashInfo.DepotName, TEXT( "" ) ) );
 				}
 
 				// Iterate through all module and see if we have dll and pdb associated with the module, if so add it to the files to sync.
@@ -354,22 +353,22 @@ bool ICrashDebugHelper::SyncModules()
 				// Now, sync all files.
 				for( const auto& Filename : FilesToSync )
 				{
-					const FString DepotPath = DepotName + Filename;
+					const FString DepotPath = CrashInfo.DepotName + Filename;
 					if( Label->Sync( DepotPath ) )
 					{
-						UE_LOG( LogCrashDebugHelper, Warning, TEXT( " ... synced binary '%s'." ), *DepotPath );
+						UE_LOG( LogCrashDebugHelper, Warning, TEXT( "Synced binary: %s" ), *DepotPath );
 					}
 					FilesToBeCached.Add( DepotPath );
 				}
 			}
 
 			// Initialize and add a new PDB Cache entry to the database.
-			CrashInfo.PDBCacheEntry = PDBCache.CreateAndAddPDBCacheEntry( CrashInfo.LabelName, DepotRoot, DepotName, FilesToBeCached );
+			CrashInfo.PDBCacheEntry = FPDBCache::Get().CreateAndAddPDBCacheEntry( CrashInfo.LabelName, DepotRoot, CrashInfo.DepotName, FilesToBeCached );
 		}
 	}
 	else
 	{
-		UE_LOG( LogCrashDebugHelper, Error, TEXT( "Could not find label '%s'."), *CrashInfo.LabelName );
+		UE_LOG( LogCrashDebugHelper, Error, TEXT( "Could not find label: %s"), *CrashInfo.LabelName );
 		return false;
 	}
 
@@ -385,23 +384,15 @@ bool ICrashDebugHelper::SyncSourceFile()
 	}
 
 	// Sync a single source file to requested CL.
-	FString DepotPath = DepotName / CrashInfo.SourceFile + TEXT("@") + TTypeToString<int32>::ToString(CrashInfo.BuiltFromCL);
+	FString DepotPath = CrashInfo.DepotName / CrashInfo.SourceFile + TEXT( "@" ) + TTypeToString<int32>::ToString( CrashInfo.BuiltFromCL );
 	ISourceControlModule::Get().GetProvider().Execute(ISourceControlOperation::Create<FSync>(), DepotPath);
 
-	UE_LOG( LogCrashDebugHelper, Warning, TEXT( "Syncing a single source file '%s' to CL %i."), *DepotPath, CrashInfo.BuiltFromCL );
+	UE_LOG( LogCrashDebugHelper, Warning, TEXT( "Syncing a single source file: %s"), *DepotPath );
 
 	return true;
 }
 
-/**
- *	Load the given ANSI text file to an array of strings - one FString per line of the file.
- *	Intended for use in simple text parsing actions
- *
- *	@param	InFilename			The text file to read, full path
- *	@param	OutStrings			The array of FStrings to fill in
- *
- *	@return	bool				true if successful, false if not
- */
+
 bool ICrashDebugHelper::ReadSourceFile( const TCHAR* InFilename, TArray<FString>& OutStrings )
 {
 	FString Line;
@@ -419,9 +410,6 @@ bool ICrashDebugHelper::ReadSourceFile( const TCHAR* InFilename, TArray<FString>
 	}
 }
 
-/** 
- * Add adjacent lines of the source file the crash occurred in the crash report
- */
 void ICrashDebugHelper::AddSourceToReport()
 {
 	if( CrashInfo.SourceFile.Len() > 0 && CrashInfo.SourceLineNumber != 0 )
@@ -430,10 +418,10 @@ void ICrashDebugHelper::AddSourceToReport()
 		FString FullPath = FString( TEXT( "../../../" ) ) + CrashInfo.SourceFile;
 		ReadSourceFile( *FullPath, Lines );
 
-		uint64 MinLine = FMath::Clamp( CrashInfo.SourceLineNumber - 15, ( uint64 )1, ( uint64 )Lines.Num() );
-		uint64 MaxLine = FMath::Clamp( CrashInfo.SourceLineNumber + 15, ( uint64 )1, ( uint64 )Lines.Num() );
+		const uint32 MinLine = FMath::Clamp( CrashInfo.SourceLineNumber - 15, (uint32)1, (uint32)Lines.Num() );
+		const uint32 MaxLine = FMath::Clamp( CrashInfo.SourceLineNumber + 15, (uint32)1, (uint32)Lines.Num() );
 
-		for( int32 Line = MinLine; Line < MaxLine; Line++ )
+		for( uint32 Line = MinLine; Line < MaxLine; Line++ )
 		{
 			if( Line == CrashInfo.SourceLineNumber - 1 )
 			{
@@ -447,9 +435,6 @@ void ICrashDebugHelper::AddSourceToReport()
 	}
 }
 
-/** 
- * Add source control annotated adjacent lines of the source file the crash occurred in the crash report
- */
 bool ICrashDebugHelper::AddAnnotatedSourceToReport()
 {
 	// Make sure we have a source file to interrogate
@@ -462,16 +447,16 @@ bool ICrashDebugHelper::AddAnnotatedSourceToReport()
 		}
 
 		// Ask source control to annotate the file for us
-		FString DepotPath = DepotName / CrashInfo.SourceFile;
+		FString DepotPath = CrashInfo.DepotName / CrashInfo.SourceFile;
 
 		TArray<FAnnotationLine> Lines;
 		SourceControlHelpers::AnnotateFile( ISourceControlModule::Get().GetProvider(), CrashInfo.LabelName, DepotPath, Lines );
 
-		uint64 MinLine = FMath::Clamp( CrashInfo.SourceLineNumber - 15, ( uint64 )1, ( uint64 )Lines.Num() );
-		uint64 MaxLine = FMath::Clamp( CrashInfo.SourceLineNumber + 15, ( uint64 )1, ( uint64 )Lines.Num() );
+		uint32 MinLine = FMath::Clamp( CrashInfo.SourceLineNumber - 15, (uint32)1, (uint32)Lines.Num() );
+		uint32 MaxLine = FMath::Clamp( CrashInfo.SourceLineNumber + 15, (uint32)1, (uint32)Lines.Num() );
 
 		// Display a source context in the report, and decorate each line with the last editor of the line
-		for( int32 Line = MinLine; Line < MaxLine; Line++ )
+		for( uint32 Line = MinLine; Line < MaxLine; Line++ )
 		{			
 			if( Line == CrashInfo.SourceLineNumber )
 			{
@@ -488,18 +473,13 @@ bool ICrashDebugHelper::AddAnnotatedSourceToReport()
 	return false;
 }
 
-/**
- * Add a line to the report
- */
 void FCrashInfo::Log( FString Line )
 {
 	UE_LOG( LogCrashDebugHelper, Warning, TEXT("%s"), *Line );
 	Report += Line + LINE_TERMINATOR;
 }
 
-/** 
- * Convert the processor architecture to a human readable string
- */
+
 const TCHAR* FCrashInfo::GetProcessorArchitecture( EProcessorArchitecture PA )
 {
 	switch( PA )
@@ -515,9 +495,7 @@ const TCHAR* FCrashInfo::GetProcessorArchitecture( EProcessorArchitecture PA )
 	return TEXT( "Unknown" );
 }
 
-/** 
- * Calculate the byte size of a UTF-8 string
- */
+
 int64 FCrashInfo::StringSize( const ANSICHAR* Line )
 {
 	int64 Size = 0;
@@ -531,9 +509,7 @@ int64 FCrashInfo::StringSize( const ANSICHAR* Line )
 	return Size;
 }
 
-/** 
- * Write a line of UTF-8 to a file
- */
+
 void FCrashInfo::WriteLine( FArchive* ReportFile, const ANSICHAR* Line )
 {
 	if( Line != NULL )
@@ -545,9 +521,7 @@ void FCrashInfo::WriteLine( FArchive* ReportFile, const ANSICHAR* Line )
 	ReportFile->Serialize( TCHAR_TO_UTF8( LINE_TERMINATOR ), FCStringWide::Strlen(LINE_TERMINATOR) );
 }
 
-/** 
- * Write all the data mined from the minidump to a text file
- */
+
 void FCrashInfo::GenerateReport( const FString& DiagnosticsPath )
 {
 	FArchive* ReportFile = IFileManager::Get().CreateFileWriter( *DiagnosticsPath );
@@ -558,9 +532,9 @@ void FCrashInfo::GenerateReport( const FString& DiagnosticsPath )
 		WriteLine( ReportFile, TCHAR_TO_UTF8( TEXT( "Generating report for minidump" ) ) );
 		WriteLine( ReportFile );
 
-		if ( ProductVersion.Len() > 0 )
+		if ( EngineVersion.Len() > 0 )
 		{
-			Line = FString::Printf( TEXT( "Application version %s" ), *ProductVersion );
+			Line = FString::Printf( TEXT( "Application version %s" ), *EngineVersion );
 			WriteLine( ReportFile, TCHAR_TO_UTF8( *Line ) );
 		}
 		else if( Modules.Num() > 0 )
@@ -650,565 +624,126 @@ void FCrashInfo::GenerateReport( const FString& DiagnosticsPath )
 
 bool ICrashDebugHelper::SyncRequiredFilesForDebuggingFromLabel(const FString& InLabel, const FString& InPlatform)
 {
-	// Ensure we are in a valid state to sync
-	if (bInitialized == false)
-	{
-		UE_LOG(LogCrashDebugHelper, Warning, TEXT("SyncRequiredFilesForDebuggingFromLabel: CrashDebugHelper is not initialized properly."));
-		return false;
-	}
-
-	if (InLabel.Len() <= 0)
-	{
-		UE_LOG(LogCrashDebugHelper, Warning, TEXT("SyncRequiredFilesForDebuggingFromLabel: Invalid Label parameter."));
-		return false;
-	}
-
-	if (InPlatform.Len() <= 0)
-	{
-		UE_LOG(LogCrashDebugHelper, Warning, TEXT("SyncRequiredFilesForDebuggingFromLabel: Invalid Platform parameter."));
-		return false;
-	}
-
-	bool bSyncFiles = true;
-
-	// We have a valid label... 
-	// This command will get the list of all Win64 pdb files under engine at the given label
-	//     p4 files //depot/UE4/Engine/Binaries/Win64/...pdb...@UE4_[2012-10-24_04.00]
-	// This command will get the list of all Win64 pdb files under game folders at the given label
-	//     p4 files //depot/UE4/...Game/Binaries/Win64/...pdb...@UE4_[2012-10-24_04.00]
-	FString EngineRoot = FString::Printf(TEXT("%s/Engine/Binaries/%s/"), *DepotName, *InPlatform);
-	FString EnginePdbFiles = EngineRoot + TEXT("...pdb...");
-	FString EngineExeFiles = EngineRoot + TEXT("...exe...");
-	FString EngineDllFiles = EngineRoot + TEXT("...dll...");
-
-	ISourceControlProvider& SourceControlProvider = ISourceControlModule::Get().GetProvider();
-	TSharedPtr<ISourceControlLabel> Label = SourceControlProvider.GetLabel(InLabel);
-	if(Label.IsValid())
-	{
-		TArray< TSharedRef<ISourceControlRevision, ESPMode::ThreadSafe> > EnginePdbList;
-		TArray< TSharedRef<ISourceControlRevision, ESPMode::ThreadSafe> > EngineExeList;
-		TArray< TSharedRef<ISourceControlRevision, ESPMode::ThreadSafe> > EngineDllList;
-
-		bool bGotEngineFiles = true;
-		bGotEngineFiles |= Label->GetFileRevisions(EnginePdbFiles, EnginePdbList);
-		bGotEngineFiles |= Label->GetFileRevisions(EngineExeFiles, EngineExeList);
-		bGotEngineFiles |= Label->GetFileRevisions(EngineDllFiles, EngineDllList);
-
-		FString GameRoot = FString::Printf(TEXT("%s/...Game/Binaries/%s/"), *DepotName, *InPlatform);
-		FString GamePdbFiles = GameRoot + TEXT("...pdb...");
-		FString GameExeFiles = GameRoot + TEXT("...exe...");
-		FString GameDllFiles = GameRoot + TEXT("...dll...");
-		TArray< TSharedRef<ISourceControlRevision, ESPMode::ThreadSafe> > GamePdbList;
-		TArray< TSharedRef<ISourceControlRevision, ESPMode::ThreadSafe> > GameExeList;
-		TArray< TSharedRef<ISourceControlRevision, ESPMode::ThreadSafe> > GameDllList;
-
-		bool bGotGameFiles = true;
-		bGotGameFiles |= Label->GetFileRevisions(GamePdbFiles, GamePdbList);
-		bGotGameFiles |= Label->GetFileRevisions(GameExeFiles, GameExeList);
-		bGotGameFiles |= Label->GetFileRevisions(GameDllFiles, GameDllList);
-
-		if (bGotEngineFiles == true)
-		{
-			TArray< TSharedRef<ISourceControlRevision, ESPMode::ThreadSafe> > CopyFileList;
-			CopyFileList += EnginePdbList;
-			CopyFileList += EngineExeList;
-			CopyFileList += EngineDllList;
-			CopyFileList += GamePdbList;
-			CopyFileList += GameExeList;
-			CopyFileList += GameDllList;
-
-			int32 CopyCount = 0;
-
-			// Copy all the files retrieved
-			for (int32 FileIdx = 0; FileIdx < CopyFileList.Num(); FileIdx++)
-			{
-				TSharedRef<ISourceControlRevision, ESPMode::ThreadSafe> Revision = CopyFileList[FileIdx];
-
-				// Copy the files to a flat directory.
-				// This will have problems if there are any files named the same!!!!
-				FString LocalStoreFolder = LocalSymbolStore / InPlatform;
-				FString CopyFilename = LocalStoreFolder / FPaths::GetCleanFilename(Revision->GetFilename());
-
-//				UE_LOG(LogCrashDebugHelper, Display, TEXT("Copying engine file: %s to %s"), *Filename, *CopyFilename);
-
-				if (Revision->Get(CopyFilename) == true)
-				{
-					CopyCount++;
-				}
-			}
-
-			//@todo. Should we verify EVERY file was copied?
-			return (CopyCount > 0);
-		}
-	}
-	else
-	{
-		UE_LOG(LogCrashDebugHelper, Warning, TEXT("SyncRequiredFiles: Invalid label specified: %s"), *InLabel);
-	}
+	// @TODO yrx 2015-02-19 Use PDB cache
 
 	return false;
 }
 
 bool ICrashDebugHelper::SyncRequiredFilesForDebuggingFromChangelist(int32 InChangelistNumber, const FString& InPlatform)
 {
-	//@todo. Allow for syncing a changelist directly?
-	// Not really useful as the source indexed pdbs will be tied to labeled builds.
-	// For now, we will not support this
-
-	const FString BuildLabel = RetrieveBuildLabel(InChangelistNumber);
-	if (BuildLabel.Len() > 0)
-	{
-		return SyncRequiredFilesForDebuggingFromLabel(BuildLabel, InPlatform);
-	}
+	// @TODO yrx 2015-02-19 Use PDB cache
 
-	UE_LOG(LogCrashDebugHelper, Warning, 
-		TEXT("SyncRequiredFilesForDebuggingFromChangelist: Failed to find label for changelist %d"), InChangelistNumber);
 	return false;
 }
 
-/**
- *	Retrieve the build label for the given engine version or changelist number.
- *
- *	@param	InChangelistNumber	The changelist number to retrieve the label for
- *	@return	FString				The label if successful, empty if it wasn't found or another error
- */
-FString ICrashDebugHelper::RetrieveBuildLabel(int32 InChangelistNumber)
-{
-	FString FoundLabelString;
-
-	if( InChangelistNumber <= 0 )
-	{
-		UE_LOG(LogCrashDebugHelper, Warning, TEXT("RetrieveBuildLabel: Invalid parameters."));
-		return FoundLabelString;
-	}
-
-	// Try to find the label directly in source control by using the pattern supplied via ini
-	if ( FoundLabelString.IsEmpty() && !SourceControlBuildLabelPattern.IsEmpty() )
-	{
-		const FString ChangelistString = FString::Printf(TEXT("%d"), InChangelistNumber);
-		const FString TestLabel = SourceControlBuildLabelPattern.Replace(TEXT("%CHANGELISTNUMBER%"), *ChangelistString, ESearchCase::CaseSensitive );
-		TArray< TSharedRef<ISourceControlLabel> > Labels = ISourceControlModule::Get().GetProvider().GetLabels( TestLabel );
-		if ( Labels.Num() > 0 )
-		{
-			const int32 LabelIndex = 0;
-			// If we found more than one label, warn about it and just use the first one
-			if ( Labels.Num() > 1 )
-			{
-				UE_LOG(LogCrashDebugHelper, Warning, TEXT("RetrieveBuildLabel: More than one build label found with pattern %s - Using label %s"), *TestLabel, *Labels[LabelIndex]->GetName());
-			}
-
-			FoundLabelString = Labels[LabelIndex]->GetName();
-			UE_LOG(LogCrashDebugHelper, Log, TEXT("RetrieveBuildLabel: Found label %s matching pattern %s in source control."), *FoundLabelString, *TestLabel);
-		}
-	}
-
-	return FoundLabelString;
-}
-
-void ICrashDebugHelper::RetrieveBuildLabelAndNetworkPaths( int32 InChangelistNumber )
+void ICrashDebugHelper::FindSymbolsAndBinariesStorage()
 {
 	CrashInfo.ExecutablesPath.Empty();
 	CrashInfo.SymbolsPath.Empty();
 	CrashInfo.LabelName.Empty();
 
-	if( InChangelistNumber <= 0 )
-	{
-		UE_LOG( LogCrashDebugHelper, Warning, TEXT( "RetrieveBuildLabelAndNetworkPaths: Invalid parameters." ) );
-	}
-
-	// Try to find the network path by using the pattern supplied via ini.
-	// If this step successes, we will grab the executable from the network path instead of P4.
-	// At this moment we only care about UE4 releases.
-	const bool bCanUseNetworkPath = DepotName.Contains( TEXT( "UE4-Releases" ) );
-	bool bFound = false;
-	if( bCanUseNetworkPath && !ExecutablePathPattern.IsEmpty() )
-	{
-		const FString TestNetworkExecutablesPath = ExecutablePathPattern.Replace( TEXT( "%PRODUCT_VERSION%" ), *CrashInfo.ProductVersion );
-		const FString TestNetworkSymbolsPath = TestNetworkExecutablesPath + TEXT("Symbols");
-		const bool bHasExecutables = IFileManager::Get().DirectoryExists( *TestNetworkExecutablesPath );
-		const bool bHasSymbols = IFileManager::Get().DirectoryExists( *TestNetworkExecutablesPath );
-
-		if( bHasExecutables && bHasSymbols )
-		{
-			CrashInfo.ExecutablesPath = TestNetworkExecutablesPath;
-			CrashInfo.SymbolsPath = TestNetworkSymbolsPath;
-			bFound = true;
-			UE_LOG( LogCrashDebugHelper, Log, TEXT( "RetrieveBuildLabelAndNetworkPaths: Found network path for executables and symbols '%s'." ), *CrashInfo.ExecutablesPath );
-		}
-		else if( bHasExecutables )
-		{
-			CrashInfo.ExecutablesPath = TestNetworkExecutablesPath;
-			UE_LOG( LogCrashDebugHelper, Log, TEXT( "RetrieveBuildLabelAndNetworkPaths: Found network path for executables '%s'." ), *CrashInfo.ExecutablesPath );
-		}
-		else
-		{
-			UE_LOG( LogCrashDebugHelper, Log, TEXT( "RetrieveBuildLabelAndNetworkPaths: Network path for executables '%s' not found." ), *TestNetworkExecutablesPath );
-		}
-	}
-
-	// Try to find the label directly in source control by using the pattern supplied via ini.
-	if( !bFound && !SourceControlBuildLabelPattern.IsEmpty() )
-	{
-		const FString ChangelistString = FString::Printf( TEXT( "%d" ), InChangelistNumber );
-		const FString TestLabelWithCL = SourceControlBuildLabelPattern.Replace( TEXT( "%CHANGELISTNUMBER%" ), *ChangelistString, ESearchCase::CaseSensitive );
-		TArray< TSharedRef<ISourceControlLabel> > Labels = ISourceControlModule::Get().GetProvider().GetLabels( TestLabelWithCL );
-		if( Labels.Num() > 0 )
-		{
-			const int32 LabelIndex = 0;
-			// If we found more than one label, warn about it and just use the first one
-			if( Labels.Num() > 1 )
-			{
-				UE_LOG( LogCrashDebugHelper, Warning, TEXT( "RetrieveBuildLabelAndNetworkPaths: More than one build label found with pattern '%s' - Using label '%s'" ), *TestLabelWithCL, *Labels[LabelIndex]->GetName() );
-			}
-
-			CrashInfo.LabelName = Labels[LabelIndex]->GetName();
-			UE_LOG( LogCrashDebugHelper, Log, TEXT( "RetrieveBuildLabelAndNetworkPaths: Found label '%s' matching pattern '%s' in source control." ), *CrashInfo.LabelName, *TestLabelWithCL );
-		}
-	}
-}
-
-/*-----------------------------------------------------------------------------
-	PDB Cache implementation
------------------------------------------------------------------------------*/
-
-const TCHAR* FPDBCache::PDBTimeStampFileNoMeta = TEXT( "PDBTimeStamp.txt" );
-const TCHAR* FPDBCache::PDBTimeStampFile = TEXT( "PDBTimeStamp.bin" );
-
-void FPDBCache::Init()
-{
-	// PDB Cache
-	// Default configuration
-	//PDBCachePath=F:/CrashReportPDBCache/
-	//DepotRoot=F:/depot
-	//DaysToDeleteUnusedFilesFromPDBCache=7
-	//PDBCacheSizeGB=128
-	//MinDiskFreeSpaceGB=256
-
-	// Can be enabled only through the command line.
-	FParse::Bool( FCommandLine::Get(), TEXT( "bUsePDBCache=" ), bUsePDBCache );
-
-	UE_LOG( LogCrashDebugHelper, Warning, TEXT( "bUsePDBCache is %s" ), bUsePDBCache ? TEXT( "enabled" ) : TEXT( "disabled" ) );
-
-	// Get the rest of the PDB cache configuration.
-	if( bUsePDBCache )
-	{
-		if( !GConfig->GetString( TEXT( "Engine.CrashDebugHelper" ), TEXT( "PDBCachePath" ), PDBCachePath, GEngineIni ) )
-		{
-			UE_LOG( LogCrashDebugHelper, Warning, TEXT( "Failed to get PDBCachePath from ini file... PDB Cache disabled" ) );
-			bUsePDBCache = false;
-		}
-	}
-
-	if( bUsePDBCache )
-	{
-		if( !GConfig->GetInt( TEXT( "Engine.CrashDebugHelper" ), TEXT( "PDBCacheSizeGB" ), PDBCacheSizeGB, GEngineIni ) )
-		{
-			UE_LOG( LogCrashDebugHelper, Warning, TEXT( "Failed to get PDBCachePath from ini file... Using default value" ) );
-		}
-
-		if( !GConfig->GetInt( TEXT( "Engine.CrashDebugHelper" ), TEXT( "MinDiskFreeSpaceGB" ), MinDiskFreeSpaceGB, GEngineIni ) )
-		{
-			UE_LOG( LogCrashDebugHelper, Warning, TEXT( "Failed to get MinDiskFreeSpaceGB from ini file... Using default value" ) );
-		}
-
-		if( !GConfig->GetInt( TEXT( "Engine.CrashDebugHelper" ), TEXT( "DaysToDeleteUnusedFilesFromPDBCache" ), DaysToDeleteUnusedFilesFromPDBCache, GEngineIni ) )
-		{
-			UE_LOG( LogCrashDebugHelper, Warning, TEXT( "Failed to get DaysToDeleteUnusedFilesFromPDBCache from ini file... Using default value" ) );
-		}
-
-		InitializePDBCache();
-		CleanPDBCache( DaysToDeleteUnusedFilesFromPDBCache );
-
-		// Verify that we have enough space to enable the PDB Cache.
-		uint64 TotalNumberOfBytes = 0;
-		uint64 NumberOfFreeBytes = 0;
-		FPlatformMisc::GetDiskTotalAndFreeSpace( PDBCachePath, TotalNumberOfBytes, NumberOfFreeBytes );
-
-		const int32 TotalDiscSpaceGB = int32( TotalNumberOfBytes >> 30 );
-		const int32 DiskFreeSpaceGB = int32( NumberOfFreeBytes >> 30 );
-
-		if( DiskFreeSpaceGB < MinDiskFreeSpaceGB )
-		{
-			// There is not enough free space, calculate the current PDB cache usage and try removing the old data.
-			const int32 CurrentPDBCacheSizeGB = GetPDBCacheSizeGB();
-			const int32 DiskFreeSpaceAfterCleanGB = DiskFreeSpaceGB + CurrentPDBCacheSizeGB;
-			if( DiskFreeSpaceAfterCleanGB < MinDiskFreeSpaceGB )
-			{
-				UE_LOG( LogCrashDebugHelper, Error, TEXT( "There is not enough free space. PDB Cache disabled." ) );
-				UE_LOG( LogCrashDebugHelper, Error, TEXT( "Current disk free space is %i GBs." ), DiskFreeSpaceGB );
-				UE_LOG( LogCrashDebugHelper, Error, TEXT( "To enable the PDB Cache you need to free %i GB of space" ), MinDiskFreeSpaceGB - DiskFreeSpaceAfterCleanGB );
-				bUsePDBCache = false;
-				// Remove all data.
-				CleanPDBCache( 0 );
-			}
-			else
-			{
-				// Clean the PDB cache until we get enough free space.
-				CleanPDBCache( DaysToDeleteUnusedFilesFromPDBCache, MinDiskFreeSpaceGB - DiskFreeSpaceGB );
-			}
-		}
-	}
-
-	if( bUsePDBCache )
-	{
-		UE_LOG( LogCrashDebugHelper, Log, TEXT( "PDBCachePath=%s" ), *PDBCachePath );
-		UE_LOG( LogCrashDebugHelper, Log, TEXT( "PDBCacheSizeGB=%i" ), PDBCacheSizeGB );
-		UE_LOG( LogCrashDebugHelper, Log, TEXT( "MinDiskFreeSpaceGB=%i" ), MinDiskFreeSpaceGB );
-		UE_LOG( LogCrashDebugHelper, Log, TEXT( "DaysToDeleteUnusedFilesFromPDBCache=%i" ), DaysToDeleteUnusedFilesFromPDBCache );
-	}
-}
-
-
-void FPDBCache::InitializePDBCache()
-{
-	const double StartTime = FPlatformTime::Seconds();
-	IFileManager::Get().MakeDirectory( *PDBCachePath, true );
-
-	TArray<FString> PDBCacheEntryDirectories;
-	IFileManager::Get().FindFiles( PDBCacheEntryDirectories, *PDBCachePath, false, true );
-
-	for( const auto& Directory : PDBCacheEntryDirectories )
+	if( CrashInfo.BuiltFromCL == FCrashInfo::INVALID_CHANGELIST )
 	{
-		FPDBCacheEntryRef Entry = ReadPDBCacheEntry( Directory );
-		PDBCacheEntries.Add( Directory, Entry );
+		UE_LOG( LogCrashDebugHelper, Warning, TEXT( "Invalid parameters" ) );
+		return;
 	}
 
-	SortPDBCache();
+	UE_LOG( LogCrashDebugHelper, Log, TEXT( "Engine version: %s" ), *CrashInfo.EngineVersion );
 
-	const double TotalTime = FPlatformTime::Seconds() - StartTime;
-	UE_LOG( LogCrashDebugHelper, Log, TEXT( "PDB Cache initialized in %.2f ms" ), TotalTime*1000.0f );
-	UE_LOG( LogCrashDebugHelper, Log, TEXT( "Found %i entries which occupy %i GBs" ), PDBCacheEntries.Num(), GetPDBCacheSizeGB() );
-}
-
-void FPDBCache::CleanPDBCache( int32 DaysToDelete, int32 NumberOfGBsToBeCleaned /*= 0 */ )
-{
-	// Not very efficient, but should do the trick.
-	// Revisit it later.
-	const double StartTime = FPlatformTime::Seconds();
-
-	TSet<FString> EntriesToBeRemoved;
-
-	// Find all outdated PDB Cache entries and mark them for removal.
-	const double DaysToDeleteAsSeconds = FTimespan( DaysToDelete, 0, 0, 0 ).GetTotalSeconds();
-	int32 NumGBsCleaned = 0;
-	for( const auto& It : PDBCacheEntries )
+	int32 Index = 0;
+	bool bFoundPattern = false;
+	FString ExecutablePathPattern;
+	FString SymbolPathPattern;
+	// Find branch.
+	for( ; Index < Branches.Num(); Index++ )
 	{
-		const FPDBCacheEntryRef& Entry = It.Value;
-		const FString EntryDirectory = PDBCachePath / Entry->Directory;
-		const FString EntryTimeStampFilename = EntryDirectory / PDBTimeStampFile;
-
-		const double EntryFileAge = IFileManager::Get().GetFileAgeSeconds( *EntryTimeStampFilename );
-		if( EntryFileAge > DaysToDeleteAsSeconds )
+		if( CrashInfo.DepotName.Contains( Branches[Index] ) )
 		{
-			EntriesToBeRemoved.Add( Entry->Directory );
-			NumGBsCleaned += Entry->SizeGB;
+			bFoundPattern = true;
+			ExecutablePathPattern = ExecutablePathPatterns[Index];
+			SymbolPathPattern = SymbolPathPatterns[Index];
+			break;
 		}
 	}
 
-	if( NumberOfGBsToBeCleaned > 0 && NumGBsCleaned < NumberOfGBsToBeCleaned )
+	if( bFoundPattern )
 	{
-		// Do the second pass if we need to remove more PDB Cache entries due to the free disk space restriction.
-		for( const auto& It : PDBCacheEntries )
-		{
-			const FPDBCacheEntryRef& Entry = It.Value;
-
-			if( !EntriesToBeRemoved.Contains( Entry->Directory ) )
-			{
-				EntriesToBeRemoved.Add( Entry->Directory );
-				NumGBsCleaned += Entry->SizeGB;
-
-				if( NumGBsCleaned > NumberOfGBsToBeCleaned )
-				{
-					// Break the loop, we are done.
-					break;
-				}
-			}
-		}
+		UE_LOG( LogCrashDebugHelper, Log, TEXT( "Using branch: %s" ), *CrashInfo.DepotName );
 	}
-
-	// Remove all marked PDB Cache entries.
-	for( const auto& EntryDirectory : EntriesToBeRemoved )
-	{
-		RemovePDBCacheEntry( EntryDirectory );
-	}
-
-	const double TotalTime = FPlatformTime::Seconds() - StartTime;
-	UE_LOG( LogCrashDebugHelper, Log, TEXT( "PDB Cache cleaned %i GBs in %.2f ms" ), NumGBsCleaned, TotalTime*1000.0f );
-}
-
-FPDBCacheEntryRef FPDBCache::CreateAndAddPDBCacheEntry( const FString& OriginalLabelName, const FString& DepotRoot, const FString& DepotName, const TArray<FString>& FilesToBeCached )
-{
-	const FString CleanedLabelName = EscapePath( OriginalLabelName );
-	const FString EntryDirectory = PDBCachePath / CleanedLabelName;
-	const FString EntryTimeStampFilename = EntryDirectory / PDBTimeStampFile;
-
-	const FString LocalDepotDir = DepotRoot / DepotName.Replace( P4_DEPOT_PREFIX, TEXT( "" ) );
-
-	UE_LOG( LogCrashDebugHelper, Warning, TEXT( "PDB Cache entry '%s' is being copied from '%s', it will take some time" ), *CleanedLabelName, *OriginalLabelName );
-	for( const auto& Filename : FilesToBeCached )
-	{
-		const FString SourceDirectoryWithSearch = Filename.Replace( *DepotName, *LocalDepotDir );
-
-		TArray<FString> MatchedFiles;
-		IFileManager::Get().FindFiles( MatchedFiles, *SourceDirectoryWithSearch, true, false );
-
-		for( const auto& MatchedFilename : MatchedFiles )
-		{
-			const FString SrcFilename = FPaths::GetPath( SourceDirectoryWithSearch ) / MatchedFilename;
-			const FString DestFilename = EntryDirectory / SrcFilename.Replace( *LocalDepotDir, TEXT("") );
-			IFileManager::Get().Copy( *DestFilename, *SrcFilename );
-		}
-	}
-
-
-	TArray<FString> CachedFiles;
-	IFileManager::Get().FindFilesRecursive( CachedFiles, *EntryDirectory, TEXT( "*.*" ), true, false );
-
-	// Calculate the size of this PDB Cache entry.
-	int64 TotalSize = 0;
-	for( const auto& Filename : CachedFiles )
+	else
 	{
-		const int64 FileSize = IFileManager::Get().FileSize( *Filename );
-		TotalSize += FileSize;
+		UE_LOG( LogCrashDebugHelper, Log, TEXT( "Branch not found: %s" ), *CrashInfo.DepotName );
+		return;
 	}
 
-	// Round-up the size.
-	int32 SizeGB = (int32)FMath::DivideAndRoundUp( TotalSize, (int64)NUM_BYTES_PER_GB );
-	FPDBCacheEntryRef NewCacheEntry = MakeShareable( new FPDBCacheEntry( CachedFiles, CleanedLabelName, FDateTime::Now(), SizeGB ) );
-
-	// Verify there is an entry timestamp file, write the size of a PDB cache to avoid time consuming FindFilesRecursive during initialization.
-	TAutoPtr<FArchive> FileWriter( IFileManager::Get().CreateFileWriter( *EntryTimeStampFilename ) );
-	UE_CLOG( !FileWriter, LogCrashDebugHelper, Fatal, TEXT( "Couldn't save the timestamp file to %s" ), *EntryTimeStampFilename );
-	*FileWriter << *NewCacheEntry;
+	// %ENGINE_VERSION% - Engine versions ie.: 4.7.0-2449961+UE4-Releases+4.7
+	// %PLATFORM_NAME% - Platform name ie.: WindowsNoEditor, Win64 etc.
+	// %UT_ENGINE_VERSION% - UT ie.: ++depot+UE4-UT-CL-2454691
+	//\\epicgames.net\root\Builds\UnrealTournament\++depot+UE4-UT-CL-2417639
+	//\\epicgames.net\root\Builds\Rocket\Automated\4.7.0-2449961+++depot+UE4-Releases+4.7
+	//\\epicgames.net\root\Builds\UnrealEngineLauncher
 	
-	PDBCacheEntries.Add( CleanedLabelName, NewCacheEntry );
-	SortPDBCache();
+	const FString StrENGINE_VERSION = CrashInfo.EngineVersion;
+	const FString StrPLATFORM_NAME = TEXT( "" ); // Not implemented yet
+	const FString StrUT_ENGINE_VERSION = FString::Printf( TEXT( "++depot+UE4-UT-CL-%i" ), CrashInfo.BuiltFromCL );
 
-	return NewCacheEntry;
-}
+	const FString TestExecutablesPath = ExecutablePathPattern
+		.Replace( TEXT( "%ENGINE_VERSION%" ), *StrENGINE_VERSION )
+		.Replace( TEXT( "%PLATFORM_NAME%" ), *StrPLATFORM_NAME )
+		.Replace( TEXT( "%UT_ENGINE_VERSION%" ), *StrUT_ENGINE_VERSION );
 
-FPDBCacheEntryRef FPDBCache::CreateAndAddPDBCacheEntryMixed( const FString& ProductVersion, const TMap<FString, FString>& FilesToBeCached )
-{
-	// Enable MDD to parse all minidumps regardless the branch, to fix the issue with the missing executables on the P4 due to the build system changes.
-	const FString EntryDirectory = PDBCachePath / ProductVersion;
-	const FString EntryTimeStampFilename = EntryDirectory / PDBTimeStampFile;
- 
-	UE_LOG( LogCrashDebugHelper, Warning, TEXT( "PDB Cache entry '%s' is being created from %i files, it will take some time" ), *ProductVersion, FilesToBeCached.Num() );
-	for( const auto& It : FilesToBeCached )
-	{
-		const FString& SrcFilename = It.Key;
-		const FString DestFilename = EntryDirectory / It.Value;
-		IFileManager::Get().Copy( *DestFilename, *SrcFilename );
-	}
+	const FString TestSymbolsPath = SymbolPathPattern
+		.Replace( TEXT( "%ENGINE_VERSION%" ), *StrENGINE_VERSION )
+		.Replace( TEXT( "%PLATFORM_NAME%" ), *StrPLATFORM_NAME )
+		.Replace( TEXT( "%UT_ENGINE_VERSION%" ), *StrUT_ENGINE_VERSION );
 
-	TArray<FString> CachedFiles;
-	IFileManager::Get().FindFilesRecursive( CachedFiles, *EntryDirectory, TEXT( "*.*" ), true, false );
- 
-	// Calculate the size of this PDB Cache entry.
-	int64 TotalSize = 0;
-	for( const auto& Filename : CachedFiles )
-	{
-		const int64 FileSize = IFileManager::Get().FileSize( *Filename );
-		TotalSize += FileSize;
-	}
- 
-	// Round-up the size.
-	int32 SizeGB = (int32)FMath::DivideAndRoundUp( TotalSize, (int64)NUM_BYTES_PER_GB );
-	FPDBCacheEntryRef NewCacheEntry = MakeShareable( new FPDBCacheEntry( CachedFiles, ProductVersion, FDateTime::Now(), SizeGB ) );
-
-	// Verify there is an entry timestamp file, write the size of a PDB cache to avoid time consuming FindFilesRecursive during initialization.
-	TAutoPtr<FArchive> FileWriter( IFileManager::Get().CreateFileWriter( *EntryTimeStampFilename ) );
-	UE_CLOG( !FileWriter, LogCrashDebugHelper, Fatal, TEXT( "Couldn't save the timestamp file to %s" ), *EntryTimeStampFilename );
-	*FileWriter << *NewCacheEntry;
- 
-	PDBCacheEntries.Add( ProductVersion, NewCacheEntry );
-	SortPDBCache();
+	// Try to find the network path by using the pattern supplied via ini.
+	// If this step successes, we will grab the executable from the network path instead of P4.
+	bool bFoundDirectory = false;
 	
-	return NewCacheEntry;
-}
-
-FPDBCacheEntryRef FPDBCache::ReadPDBCacheEntry( const FString& Directory )
-{
-	const FString EntryDirectory = PDBCachePath / Directory;
-	const FString EntryTimeStampFilenameNoMeta = EntryDirectory / PDBTimeStampFileNoMeta;
-	const FString EntryTimeStampFilename = EntryDirectory / PDBTimeStampFile;
-
-	// Verify there is an entry timestamp file.
-	const FDateTime LastAccessTimeNoMeta = IFileManager::Get().GetTimeStamp( *EntryTimeStampFilenameNoMeta );
-	const FDateTime LastAccessTime = IFileManager::Get().GetTimeStamp( *EntryTimeStampFilename );
+	const bool bHasExecutables = IFileManager::Get().DirectoryExists( *TestExecutablesPath );
+	const bool bHasSymbols = IFileManager::Get().DirectoryExists( *TestSymbolsPath );
 
-	FPDBCacheEntryPtr NewEntry;
-
-	if( LastAccessTime != FDateTime::MinValue() )
+	if( bHasExecutables && bHasSymbols )
 	{
-		// Read the metadata
-		TAutoPtr<FArchive> FileReader( IFileManager::Get().CreateFileReader( *EntryTimeStampFilename ) );
-		NewEntry = MakeShareable( new FPDBCacheEntry( LastAccessTime ) );
-		*FileReader << *NewEntry;
+		CrashInfo.ExecutablesPath = TestExecutablesPath;
+		CrashInfo.SymbolsPath = TestSymbolsPath;
+		bFoundDirectory = true;
+		UE_LOG( LogCrashDebugHelper, Log, TEXT( "Using path for executables and symbols: %s" ), *CrashInfo.ExecutablesPath );
 	}
-	else if( LastAccessTimeNoMeta != FDateTime::MinValue() )
+	else if( bHasExecutables )
 	{
-		// Calculate the size of this PDB Cache entry and update to the new version.
-		TArray<FString> PDBFiles;
-		IFileManager::Get().FindFilesRecursive( PDBFiles, *EntryDirectory, TEXT( "*.*" ), true, false );
-
-		// Calculate the size of this PDB Cache entry.
-		int64 TotalSize = 0;
-		for( const auto& Filename : PDBFiles )
-		{
-			const int64 FileSize = IFileManager::Get().FileSize( *Filename );
-			TotalSize += FileSize;
-		}
-
-		// Round-up the size.
-		const int32 SizeGB = (int32)FMath::DivideAndRoundUp( TotalSize, (int64)NUM_BYTES_PER_GB );
-		NewEntry = MakeShareable( new FPDBCacheEntry( PDBFiles, Directory, LastAccessTimeNoMeta, SizeGB ) );
-
-		// Remove the old file and save the metadata.
-		TAutoPtr<FArchive> FileWriter( IFileManager::Get().CreateFileWriter( *EntryTimeStampFilename ) );
-		*FileWriter << *NewEntry;
-
-		IFileManager::Get().Delete( *EntryTimeStampFilenameNoMeta );
+		CrashInfo.ExecutablesPath = TestExecutablesPath;
+		UE_LOG( LogCrashDebugHelper, Log, TEXT( "Using path for executables: %s" ), *CrashInfo.ExecutablesPath );
 	}
 	else
 	{
-		// Something wrong.
-		check( 0 );
+		UE_LOG( LogCrashDebugHelper, Log, TEXT( "Path for executables not found: %s" ), *TestExecutablesPath );
 	}
-
-	return NewEntry.ToSharedRef();
-}
-
-void FPDBCache::TouchPDBCacheEntry( const FString& Directory )
-{
-	const FString EntryDirectory = PDBCachePath / Directory;
-	const FString EntryTimeStampFilename = EntryDirectory / PDBTimeStampFile;
-
-	FPDBCacheEntryRef& Entry = PDBCacheEntries.FindChecked( Directory );
-	Entry->SetLastAccessTimeToNow();
-
-	const bool bResult = IFileManager::Get().SetTimeStamp( *EntryTimeStampFilename, Entry->LastAccessTime );
-	SortPDBCache();
-}
-
-void FPDBCache::RemovePDBCacheEntry( const FString& Directory )
-{
-	const double StartTime = FPlatformTime::Seconds();
-
-	const FString EntryDirectory = PDBCachePath / Directory;
-
-	FPDBCacheEntryRef& Entry = PDBCacheEntries.FindChecked( Directory );
-	IFileManager::Get().DeleteDirectory( *EntryDirectory, true, true );
 	
-	const double TotalTime = FPlatformTime::Seconds() - StartTime;
-	UE_LOG( LogCrashDebugHelper, Warning, TEXT( "PDB Cache entry %s removed in %.2f ms, restored %i GBs" ), *Directory, TotalTime*1000.0f, Entry->SizeGB );
 
-	PDBCacheEntries.Remove( Directory );
-}
+	// Try to find the label directly in source control by using the pattern supplied via ini.
+	if( !bFoundDirectory && !SourceControlBuildLabelPattern.IsEmpty() )
+	{	
+		const FString ChangelistString = FString::Printf( TEXT( "%d" ), CrashInfo.BuiltFromCL );
+		const FString LabelWithCL = SourceControlBuildLabelPattern.Replace( TEXT( "%CHANGELISTNUMBER%" ), *ChangelistString, ESearchCase::CaseSensitive );
+		UE_LOG( LogCrashDebugHelper, Log, TEXT( "Label matching pattern: %s" ), *LabelWithCL );
 
-FPDBCacheEntryRef FPDBCache::FindAndTouchPDBCacheEntry( const FString& PathOrLabel )
-{
-	FPDBCacheEntryRef CacheEntry = PDBCacheEntries.FindChecked( EscapePath( PathOrLabel ) );
-	TouchPDBCacheEntry( CacheEntry->Directory );
-	return CacheEntry;
+		TArray< TSharedRef<ISourceControlLabel> > Labels = ISourceControlModule::Get().GetProvider().GetLabels( LabelWithCL );
+		if( Labels.Num() > 0 )
+		{
+			const int32 LabelIndex = 0;
+			CrashInfo.LabelName = Labels[LabelIndex]->GetName();;
+		
+			// If we found more than one label, warn about it and just use the first one
+			if( Labels.Num() > 1 )
+			{
+				UE_LOG( LogCrashDebugHelper, Warning, TEXT( "More than one build label found, using label: %s" ), *LabelWithCL, *CrashInfo.LabelName );
+			}
+			else
+			{
+				UE_LOG( LogCrashDebugHelper, Log, TEXT( "Using label: %s" ), *CrashInfo.LabelName );
+			}		
+		}
+	}
 }
diff --git a/Engine/Source/Developer/CrashDebugHelper/Private/CrashDebugHelperModule.cpp b/Engine/Source/Developer/CrashDebugHelper/Private/CrashDebugHelperModule.cpp
--- a/Engine/Source/Developer/CrashDebugHelper/Private/CrashDebugHelperModule.cpp
+++ b/Engine/Source/Developer/CrashDebugHelper/Private/CrashDebugHelperModule.cpp
@@ -1,13 +1,23 @@
 // Copyright 1998-2015 Epic Games, Inc. All Rights Reserved.
 
 #include "CrashDebugHelperPrivatePCH.h"
-#include "ModuleInterface.h"
-#include "ModuleManager.h"
-#include "CrashDebugHelperModule.h"
 
 IMPLEMENT_MODULE(FCrashDebugHelperModule, CrashDebugHelper);
 DEFINE_LOG_CATEGORY(LogCrashDebugHelper);
 
+#if PLATFORM_WINDOWS
+	#include "CrashDebugHelperWindows.h"
+
+#elif PLATFORM_LINUX
+	#include "CrashDebugHelperLinux.h"
+
+#elif PLATFORM_MAC
+	#include "CrashDebugHelperMac.h"
+
+#else
+	#error "Unknown platform"
+#endif
+
 void FCrashDebugHelperModule::StartupModule()
 {
 	CrashDebugHelper = new FCrashDebugHelper();
@@ -17,7 +27,6 @@ void FCrashDebugHelperModule::StartupModule()
 	}
 }
 
-
 void FCrashDebugHelperModule::ShutdownModule()
 {
 	if (CrashDebugHelper != NULL)
diff --git a/Engine/Source/Developer/CrashDebugHelper/Private/CrashDebugHelperPrivatePCH.h b/Engine/Source/Developer/CrashDebugHelper/Private/CrashDebugHelperPrivatePCH.h
--- a/Engine/Source/Developer/CrashDebugHelper/Private/CrashDebugHelperPrivatePCH.h
+++ b/Engine/Source/Developer/CrashDebugHelper/Private/CrashDebugHelperPrivatePCH.h
@@ -2,19 +2,7 @@
 
 #pragma once
 
-#include "Core.h"
 #include "CrashDebugHelper.h"
+#include "CrashDebugHelperModule.h"
 
-DECLARE_LOG_CATEGORY_EXTERN(LogCrashDebugHelper, Log, All);
-
-#if PLATFORM_WINDOWS
-	#include "Windows/CrashDebugHelperWindows.h"
-#endif
-
-#if PLATFORM_LINUX
-	#include "Linux/CrashDebugHelperLinux.h"
-#endif
-
-#if PLATFORM_MAC
-	#include "Mac/CrashDebugHelperMac.h"
-#endif
+DECLARE_LOG_CATEGORY_EXTERN( LogCrashDebugHelper, Log, All );
diff --git a/Engine/Source/Developer/CrashDebugHelper/Private/CrashDebugPDBCache.cpp b/Engine/Source/Developer/CrashDebugHelper/Private/CrashDebugPDBCache.cpp
new file mode 100644
--- /dev/null
+++ b/Engine/Source/Developer/CrashDebugHelper/Private/CrashDebugPDBCache.cpp
@@ -0,0 +1,350 @@
+// Copyright 1998-2015 Epic Games, Inc. All Rights Reserved.
+
+#include "CrashDebugHelperPrivatePCH.h"
+#include "CrashDebugPDBCache.h"
+
+/*-----------------------------------------------------------------------------
+	PDB Cache implementation
+-----------------------------------------------------------------------------*/
+
+const TCHAR* FPDBCache::PDBTimeStampFileNoMeta = TEXT( "PDBTimeStamp.txt" );
+const TCHAR* FPDBCache::PDBTimeStampFile = TEXT( "PDBTimeStamp.bin" );
+
+void FPDBCache::Init()
+{
+	// PDB Cache
+	// Default configuration
+	//PDBCachePath=F:/CrashReportPDBCache/
+	//DepotRoot=F:/depot
+	//DaysToDeleteUnusedFilesFromPDBCache=7
+	//PDBCacheSizeGB=128
+	//MinDiskFreeSpaceGB=256
+
+	// Can be enabled only through the command line.
+	FParse::Bool( FCommandLine::Get(), TEXT( "bUsePDBCache=" ), bUsePDBCache );
+
+	UE_LOG( LogCrashDebugHelper, Warning, TEXT( "bUsePDBCache is %s" ), bUsePDBCache ? TEXT( "enabled" ) : TEXT( "disabled" ) );
+
+	// Get the rest of the PDB cache configuration.
+	if( bUsePDBCache )
+	{
+		if( !GConfig->GetString( TEXT( "Engine.CrashDebugHelper" ), TEXT( "PDBCachePath" ), PDBCachePath, GEngineIni ) )
+		{
+			UE_LOG( LogCrashDebugHelper, Warning, TEXT( "Failed to get PDBCachePath from ini file" ) );
+			bUsePDBCache = false;
+		}
+	}
+
+	if( bUsePDBCache )
+	{
+		if( !GConfig->GetInt( TEXT( "Engine.CrashDebugHelper" ), TEXT( "PDBCacheSizeGB" ), PDBCacheSizeGB, GEngineIni ) )
+		{
+			UE_LOG( LogCrashDebugHelper, Warning, TEXT( "Failed to get PDBCachePath from ini file" ) );
+		}
+
+		if( !GConfig->GetInt( TEXT( "Engine.CrashDebugHelper" ), TEXT( "MinDiskFreeSpaceGB" ), MinDiskFreeSpaceGB, GEngineIni ) )
+		{
+			UE_LOG( LogCrashDebugHelper, Warning, TEXT( "Failed to get MinDiskFreeSpaceGB from ini file" ) );
+		}
+
+		if( !GConfig->GetInt( TEXT( "Engine.CrashDebugHelper" ), TEXT( "DaysToDeleteUnusedFilesFromPDBCache" ), DaysToDeleteUnusedFilesFromPDBCache, GEngineIni ) )
+		{
+			UE_LOG( LogCrashDebugHelper, Warning, TEXT( "Failed to get DaysToDeleteUnusedFilesFromPDBCache from ini file" ) );
+		}
+
+		InitializePDBCache();
+		CleanPDBCache( DaysToDeleteUnusedFilesFromPDBCache );
+
+		// Verify that we have enough space to enable the PDB Cache.
+		uint64 TotalNumberOfBytes = 0;
+		uint64 NumberOfFreeBytes = 0;
+		FPlatformMisc::GetDiskTotalAndFreeSpace( PDBCachePath, TotalNumberOfBytes, NumberOfFreeBytes );
+
+		const int32 TotalDiscSpaceGB = int32( TotalNumberOfBytes >> 30 );
+		const int32 DiskFreeSpaceGB = int32( NumberOfFreeBytes >> 30 );
+
+		if( DiskFreeSpaceGB < MinDiskFreeSpaceGB )
+		{
+			// There is not enough free space, calculate the current PDB cache usage and try removing the old data.
+			const int32 CurrentPDBCacheSizeGB = GetPDBCacheSizeGB();
+			const int32 DiskFreeSpaceAfterCleanGB = DiskFreeSpaceGB + CurrentPDBCacheSizeGB;
+			if( DiskFreeSpaceAfterCleanGB < MinDiskFreeSpaceGB )
+			{
+				UE_LOG( LogCrashDebugHelper, Error, TEXT( "There is not enough free space. PDB Cache disabled." ) );
+				UE_LOG( LogCrashDebugHelper, Error, TEXT( "Current disk free space is %i GBs." ), DiskFreeSpaceGB );
+				UE_LOG( LogCrashDebugHelper, Error, TEXT( "To enable the PDB Cache you need to free %i GB of space" ), MinDiskFreeSpaceGB - DiskFreeSpaceAfterCleanGB );
+				bUsePDBCache = false;
+				// Remove all data.
+				CleanPDBCache( 0 );
+			}
+			else
+			{
+				// Clean the PDB cache until we get enough free space.
+				CleanPDBCache( DaysToDeleteUnusedFilesFromPDBCache, MinDiskFreeSpaceGB - DiskFreeSpaceGB );
+			}
+		}
+	}
+
+	if( bUsePDBCache )
+	{
+		UE_LOG( LogCrashDebugHelper, Log, TEXT( "PDBCachePath=%s" ), *PDBCachePath );
+		UE_LOG( LogCrashDebugHelper, Log, TEXT( "PDBCacheSizeGB=%i" ), PDBCacheSizeGB );
+		UE_LOG( LogCrashDebugHelper, Log, TEXT( "MinDiskFreeSpaceGB=%i" ), MinDiskFreeSpaceGB );
+		UE_LOG( LogCrashDebugHelper, Log, TEXT( "DaysToDeleteUnusedFilesFromPDBCache=%i" ), DaysToDeleteUnusedFilesFromPDBCache );
+	}
+}
+
+
+void FPDBCache::InitializePDBCache()
+{
+	const double StartTime = FPlatformTime::Seconds();
+	IFileManager::Get().MakeDirectory( *PDBCachePath, true );
+
+	TArray<FString> PDBCacheEntryDirectories;
+	IFileManager::Get().FindFiles( PDBCacheEntryDirectories, *PDBCachePath, false, true );
+
+	for( const auto& Directory : PDBCacheEntryDirectories )
+	{
+		FPDBCacheEntryRef Entry = ReadPDBCacheEntry( Directory );
+		PDBCacheEntries.Add( Directory, Entry );
+	}
+
+	SortPDBCache();
+
+	const double TotalTime = FPlatformTime::Seconds() - StartTime;
+	UE_LOG( LogCrashDebugHelper, Log, TEXT( "PDB Cache initialized in %.2f ms" ), TotalTime*1000.0f );
+	UE_LOG( LogCrashDebugHelper, Log, TEXT( "Found %i entries which occupy %i GBs" ), PDBCacheEntries.Num(), GetPDBCacheSizeGB() );
+}
+
+void FPDBCache::CleanPDBCache( int32 DaysToDelete, int32 NumberOfGBsToBeCleaned /*= 0 */ )
+{
+	// Not very efficient, but should do the trick.
+	// Revisit it later.
+	const double StartTime = FPlatformTime::Seconds();
+
+	TSet<FString> EntriesToBeRemoved;
+
+	// Find all outdated PDB Cache entries and mark them for removal.
+	const double DaysToDeleteAsSeconds = FTimespan( DaysToDelete, 0, 0, 0 ).GetTotalSeconds();
+	int32 NumGBsCleaned = 0;
+	for( const auto& It : PDBCacheEntries )
+	{
+		const FPDBCacheEntryRef& Entry = It.Value;
+		const FString EntryDirectory = PDBCachePath / Entry->Directory;
+		const FString EntryTimeStampFilename = EntryDirectory / PDBTimeStampFile;
+
+		const double EntryFileAge = IFileManager::Get().GetFileAgeSeconds( *EntryTimeStampFilename );
+		if( EntryFileAge > DaysToDeleteAsSeconds )
+		{
+			EntriesToBeRemoved.Add( Entry->Directory );
+			NumGBsCleaned += Entry->SizeGB;
+		}
+	}
+
+	if( NumberOfGBsToBeCleaned > 0 && NumGBsCleaned < NumberOfGBsToBeCleaned )
+	{
+		// Do the second pass if we need to remove more PDB Cache entries due to the free disk space restriction.
+		for( const auto& It : PDBCacheEntries )
+		{
+			const FPDBCacheEntryRef& Entry = It.Value;
+
+			if( !EntriesToBeRemoved.Contains( Entry->Directory ) )
+			{
+				EntriesToBeRemoved.Add( Entry->Directory );
+				NumGBsCleaned += Entry->SizeGB;
+
+				if( NumGBsCleaned > NumberOfGBsToBeCleaned )
+				{
+					// Break the loop, we are done.
+					break;
+				}
+			}
+		}
+	}
+
+	// Remove all marked PDB Cache entries.
+	for( const auto& EntryDirectory : EntriesToBeRemoved )
+	{
+		RemovePDBCacheEntry( EntryDirectory );
+	}
+
+	const double TotalTime = FPlatformTime::Seconds() - StartTime;
+	UE_LOG( LogCrashDebugHelper, Log, TEXT( "PDB Cache cleaned %i GBs in %.2f ms" ), NumGBsCleaned, TotalTime*1000.0f );
+}
+
+FPDBCacheEntryRef FPDBCache::CreateAndAddPDBCacheEntry( const FString& OriginalLabelName, const FString& DepotRoot, const FString& DepotName, const TArray<FString>& FilesToBeCached )
+{
+	const FString CleanedLabelName = EscapePath( OriginalLabelName );
+	const FString EntryDirectory = PDBCachePath / CleanedLabelName;
+	const FString EntryTimeStampFilename = EntryDirectory / PDBTimeStampFile;
+
+	const FString LocalDepotDir = DepotRoot / DepotName.Replace( ICrashDebugHelper::P4_DEPOT_PREFIX, TEXT( "" ) );
+
+	UE_LOG( LogCrashDebugHelper, Warning, TEXT( "PDB Cache entry: %s is being copied from: %s, it will take some time" ), *CleanedLabelName, *OriginalLabelName );
+	for( const auto& Filename : FilesToBeCached )
+	{
+		const FString SourceDirectoryWithSearch = Filename.Replace( *DepotName, *LocalDepotDir );
+
+		TArray<FString> MatchedFiles;
+		IFileManager::Get().FindFiles( MatchedFiles, *SourceDirectoryWithSearch, true, false );
+
+		for( const auto& MatchedFilename : MatchedFiles )
+		{
+			const FString SrcFilename = FPaths::GetPath( SourceDirectoryWithSearch ) / MatchedFilename;
+			const FString DestFilename = EntryDirectory / SrcFilename.Replace( *LocalDepotDir, TEXT( "" ) );
+			IFileManager::Get().Copy( *DestFilename, *SrcFilename );
+		}
+	}
+
+
+	TArray<FString> CachedFiles;
+	IFileManager::Get().FindFilesRecursive( CachedFiles, *EntryDirectory, TEXT( "*.*" ), true, false );
+
+	// Calculate the size of this PDB Cache entry.
+	int64 TotalSize = 0;
+	for( const auto& Filename : CachedFiles )
+	{
+		const int64 FileSize = IFileManager::Get().FileSize( *Filename );
+		TotalSize += FileSize;
+	}
+
+	// Round-up the size.
+	int32 SizeGB = (int32)FMath::DivideAndRoundUp( TotalSize, (int64)NUM_BYTES_PER_GB );
+	FPDBCacheEntryRef NewCacheEntry = MakeShareable( new FPDBCacheEntry( CachedFiles, CleanedLabelName, FDateTime::Now(), SizeGB ) );
+
+	// Verify there is an entry timestamp file, write the size of a PDB cache to avoid time consuming FindFilesRecursive during initialization.
+	TAutoPtr<FArchive> FileWriter( IFileManager::Get().CreateFileWriter( *EntryTimeStampFilename ) );
+	UE_CLOG( !FileWriter, LogCrashDebugHelper, Fatal, TEXT( "Couldn't save the timestamp for a file: %s" ), *EntryTimeStampFilename );
+	*FileWriter << *NewCacheEntry;
+
+	PDBCacheEntries.Add( CleanedLabelName, NewCacheEntry );
+	SortPDBCache();
+
+	return NewCacheEntry;
+}
+
+FPDBCacheEntryRef FPDBCache::CreateAndAddPDBCacheEntryMixed( const FString& ProductVersion, const TMap<FString, FString>& FilesToBeCached )
+{
+	// Enable MDD to parse all minidumps regardless the branch, to fix the issue with the missing executables on the P4 due to the build system changes.
+	const FString EntryDirectory = PDBCachePath / ProductVersion;
+	const FString EntryTimeStampFilename = EntryDirectory / PDBTimeStampFile;
+
+	UE_LOG( LogCrashDebugHelper, Warning, TEXT( "PDB Cache entry: %s is being created from %i files, it will take some time" ), *ProductVersion, FilesToBeCached.Num() );
+	for( const auto& It : FilesToBeCached )
+	{
+		const FString& SrcFilename = It.Key;
+		const FString DestFilename = EntryDirectory / It.Value;
+		IFileManager::Get().Copy( *DestFilename, *SrcFilename );
+	}
+
+	TArray<FString> CachedFiles;
+	IFileManager::Get().FindFilesRecursive( CachedFiles, *EntryDirectory, TEXT( "*.*" ), true, false );
+
+	// Calculate the size of this PDB Cache entry.
+	int64 TotalSize = 0;
+	for( const auto& Filename : CachedFiles )
+	{
+		const int64 FileSize = IFileManager::Get().FileSize( *Filename );
+		TotalSize += FileSize;
+	}
+
+	// Round-up the size.
+	int32 SizeGB = (int32)FMath::DivideAndRoundUp( TotalSize, (int64)NUM_BYTES_PER_GB );
+	FPDBCacheEntryRef NewCacheEntry = MakeShareable( new FPDBCacheEntry( CachedFiles, ProductVersion, FDateTime::Now(), SizeGB ) );
+
+	// Verify there is an entry timestamp file, write the size of a PDB cache to avoid time consuming FindFilesRecursive during initialization.
+	TAutoPtr<FArchive> FileWriter( IFileManager::Get().CreateFileWriter( *EntryTimeStampFilename ) );
+	UE_CLOG( !FileWriter, LogCrashDebugHelper, Fatal, TEXT( "Couldn't save the timestamp for a file: %s" ), *EntryTimeStampFilename );
+	*FileWriter << *NewCacheEntry;
+
+	PDBCacheEntries.Add( ProductVersion, NewCacheEntry );
+	SortPDBCache();
+
+	return NewCacheEntry;
+}
+
+FPDBCacheEntryRef FPDBCache::ReadPDBCacheEntry( const FString& Directory )
+{
+	const FString EntryDirectory = PDBCachePath / Directory;
+	const FString EntryTimeStampFilenameNoMeta = EntryDirectory / PDBTimeStampFileNoMeta;
+	const FString EntryTimeStampFilename = EntryDirectory / PDBTimeStampFile;
+
+	// Verify there is an entry timestamp file.
+	const FDateTime LastAccessTimeNoMeta = IFileManager::Get().GetTimeStamp( *EntryTimeStampFilenameNoMeta );
+	const FDateTime LastAccessTime = IFileManager::Get().GetTimeStamp( *EntryTimeStampFilename );
+
+	FPDBCacheEntryPtr NewEntry;
+
+	if( LastAccessTime != FDateTime::MinValue() )
+	{
+		// Read the metadata
+		TAutoPtr<FArchive> FileReader( IFileManager::Get().CreateFileReader( *EntryTimeStampFilename ) );
+		NewEntry = MakeShareable( new FPDBCacheEntry( LastAccessTime ) );
+		*FileReader << *NewEntry;
+	}
+	else if( LastAccessTimeNoMeta != FDateTime::MinValue() )
+	{
+		// Calculate the size of this PDB Cache entry and update to the new version.
+		TArray<FString> PDBFiles;
+		IFileManager::Get().FindFilesRecursive( PDBFiles, *EntryDirectory, TEXT( "*.*" ), true, false );
+
+		// Calculate the size of this PDB Cache entry.
+		int64 TotalSize = 0;
+		for( const auto& Filename : PDBFiles )
+		{
+			const int64 FileSize = IFileManager::Get().FileSize( *Filename );
+			TotalSize += FileSize;
+		}
+
+		// Round-up the size.
+		const int32 SizeGB = (int32)FMath::DivideAndRoundUp( TotalSize, (int64)NUM_BYTES_PER_GB );
+		NewEntry = MakeShareable( new FPDBCacheEntry( PDBFiles, Directory, LastAccessTimeNoMeta, SizeGB ) );
+
+		// Remove the old file and save the metadata.
+		TAutoPtr<FArchive> FileWriter( IFileManager::Get().CreateFileWriter( *EntryTimeStampFilename ) );
+		*FileWriter << *NewEntry;
+
+		IFileManager::Get().Delete( *EntryTimeStampFilenameNoMeta );
+	}
+	else
+	{
+		// Something wrong.
+		check( 0 );
+	}
+
+	return NewEntry.ToSharedRef();
+}
+
+void FPDBCache::TouchPDBCacheEntry( const FString& Directory )
+{
+	const FString EntryDirectory = PDBCachePath / Directory;
+	const FString EntryTimeStampFilename = EntryDirectory / PDBTimeStampFile;
+
+	FPDBCacheEntryRef& Entry = PDBCacheEntries.FindChecked( Directory );
+	Entry->SetLastAccessTimeToNow();
+
+	const bool bResult = IFileManager::Get().SetTimeStamp( *EntryTimeStampFilename, Entry->LastAccessTime );
+	SortPDBCache();
+}
+
+void FPDBCache::RemovePDBCacheEntry( const FString& Directory )
+{
+	const double StartTime = FPlatformTime::Seconds();
+
+	const FString EntryDirectory = PDBCachePath / Directory;
+
+	FPDBCacheEntryRef& Entry = PDBCacheEntries.FindChecked( Directory );
+	IFileManager::Get().DeleteDirectory( *EntryDirectory, true, true );
+
+	const double TotalTime = FPlatformTime::Seconds() - StartTime;
+	UE_LOG( LogCrashDebugHelper, Warning, TEXT( "PDB Cache entry %s removed in %.2f ms, restored %i GBs" ), *Directory, TotalTime*1000.0f, Entry->SizeGB );
+
+	PDBCacheEntries.Remove( Directory );
+}
+
+FPDBCacheEntryRef FPDBCache::FindAndTouchPDBCacheEntry( const FString& PathOrLabel )
+{
+	FPDBCacheEntryRef CacheEntry = PDBCacheEntries.FindChecked( EscapePath( PathOrLabel ) );
+	TouchPDBCacheEntry( CacheEntry->Directory );
+	return CacheEntry;
+}
diff --git a/Engine/Source/Developer/CrashDebugHelper/Private/Linux/CrashDebugHelperLinux.cpp b/Engine/Source/Developer/CrashDebugHelper/Private/Linux/CrashDebugHelperLinux.cpp
--- a/Engine/Source/Developer/CrashDebugHelper/Private/Linux/CrashDebugHelperLinux.cpp
+++ b/Engine/Source/Developer/CrashDebugHelper/Private/Linux/CrashDebugHelperLinux.cpp
@@ -1,10 +1,7 @@
 // Copyright 1998-2015 Epic Games, Inc. All Rights Reserved.
 
 #include "CrashDebugHelperPrivatePCH.h"
-#include "Database.h"
-
-
-#include "ISourceControlModule.h"
+#include "CrashDebugHelperLinux.h"
 
 FCrashDebugHelperLinux::FCrashDebugHelperLinux()
 {
@@ -16,40 +13,9 @@ FCrashDebugHelperLinux::~FCrashDebugHelperLinux()
 
 bool FCrashDebugHelperLinux::ParseCrashDump(const FString& InCrashDumpName, FCrashDebugInfo& OutCrashDebugInfo)
 {
-	if (bInitialized == false)
-	{
-		UE_LOG(LogCrashDebugHelper, Warning, TEXT("ParseCrashDump: CrashDebugHelper not initialized"));
-		return false;
-	}
-
-	//if (ParseCrashDump_Linux(*InCrashDumpName, OutCrashDebugInfo) == true)
-	{
-		return true;
-	}
-
 	return false;
 }
 
-bool FCrashDebugHelperLinux::SyncAndDebugCrashDump(const FString& InCrashDumpName)
-{
-	if (bInitialized == false)
-	{
-		UE_LOG(LogCrashDebugHelper, Warning, TEXT("SyncAndDebugCrashDump: CrashDebugHelper not initialized"));
-		return false;
-	}
-
-	FCrashDebugInfo CrashDebugInfo;
-	
-	// Parse the crash dump if it hasn't already been...
-	if (ParseCrashDump(InCrashDumpName, CrashDebugInfo) == false)
-	{
-		UE_LOG(LogCrashDebugHelper, Warning, TEXT("SyncAndDebugCrashDump: Failed to parse crash dump %s"), *InCrashDumpName);
-		return false;
-	}
-
-	return true;
-}
-
 bool FCrashDebugHelperLinux::CreateMinidumpDiagnosticReport( const FString& InCrashDumpName )
 {
 	return true;
diff --git a/Engine/Source/Developer/CrashDebugHelper/Private/Linux/CrashDebugHelperLinux.h b/Engine/Source/Developer/CrashDebugHelper/Private/Linux/CrashDebugHelperLinux.h
--- a/Engine/Source/Developer/CrashDebugHelper/Private/Linux/CrashDebugHelperLinux.h
+++ b/Engine/Source/Developer/CrashDebugHelper/Private/Linux/CrashDebugHelperLinux.h
@@ -2,6 +2,8 @@
 
 #pragma once
 
+#include "CrashDebugHelper.h"
+
 class FCrashDebugHelperLinux : public ICrashDebugHelper
 {
 public:
@@ -18,8 +20,6 @@ public:
 	 */
 	virtual bool ParseCrashDump(const FString& InCrashDumpName, FCrashDebugInfo& OutCrashDebugInfo) override;
 
-	virtual bool SyncAndDebugCrashDump(const FString& InCrashDumpName) override;
-
 	/**
 	 *	Parse the given crash dump, and generate a report. 
 	 *
diff --git a/Engine/Source/Developer/CrashDebugHelper/Private/Mac/CrashDebugHelperMac.cpp b/Engine/Source/Developer/CrashDebugHelper/Private/Mac/CrashDebugHelperMac.cpp
--- a/Engine/Source/Developer/CrashDebugHelper/Private/Mac/CrashDebugHelperMac.cpp
+++ b/Engine/Source/Developer/CrashDebugHelper/Private/Mac/CrashDebugHelperMac.cpp
@@ -1,14 +1,12 @@
 // Copyright 1998-2015 Epic Games, Inc. All Rights Reserved.
 
 #include "CrashDebugHelperPrivatePCH.h"
-#include "Database.h"
+#include "CrashDebugHelperMac.h"
 #include "EngineVersion.h"
 #include "ApplePlatformSymbolication.h"
 #include "CrashReporter.h"
 #include <cxxabi.h>
 
-#include "ISourceControlModule.h"
-
 static int32 ParseReportVersion(TCHAR const* CrashLog, int32& OutVersion)
 {
 	int32 Found = 0;
@@ -542,26 +540,6 @@ bool FCrashDebugHelperMac::ParseCrashDump(const FString& InCrashDumpName, FCrash
 	return false;
 }
 
-bool FCrashDebugHelperMac::SyncAndDebugCrashDump(const FString& InCrashDumpName)
-{
-	if (bInitialized == false)
-	{
-		UE_LOG(LogCrashDebugHelper, Warning, TEXT("SyncAndDebugCrashDump: CrashDebugHelper not initialized"));
-		return false;
-	}
-
-	FCrashDebugInfo CrashDebugInfo;
-	
-	// Parse the crash dump if it hasn't already been...
-	if (ParseCrashDump(InCrashDumpName, CrashDebugInfo) == false)
-	{
-		UE_LOG(LogCrashDebugHelper, Warning, TEXT("SyncAndDebugCrashDump: Failed to parse crash dump %s"), *InCrashDumpName);
-		return false;
-	}
-
-	return true;
-}
-
 bool FCrashDebugHelperMac::CreateMinidumpDiagnosticReport( const FString& InCrashDumpName )
 {
 	bool bOK = false;
@@ -606,7 +584,7 @@ bool FCrashDebugHelperMac::CreateMinidumpDiagnosticReport( const FString& InCras
 			int32 Result = ParseVersion(*CrashDump, Major, Minor, Build, CLNumber, Branch);
 			if(Result >= 3)
 			{
-				CrashInfo.ProductVersion = FEngineVersion(Major, Minor, Build, CLNumber, Branch).ToString();
+				CrashInfo.EngineVersion = FEngineVersion(Major, Minor, Build, CLNumber, Branch).ToString();
 			}
 			
 			if(Result >= 4)
diff --git a/Engine/Source/Developer/CrashDebugHelper/Private/Mac/CrashDebugHelperMac.h b/Engine/Source/Developer/CrashDebugHelper/Private/Mac/CrashDebugHelperMac.h
--- a/Engine/Source/Developer/CrashDebugHelper/Private/Mac/CrashDebugHelperMac.h
+++ b/Engine/Source/Developer/CrashDebugHelper/Private/Mac/CrashDebugHelperMac.h
@@ -2,6 +2,8 @@
 
 #pragma once
 
+#include "CrashDebugHelper.h"
+
 class FCrashDebugHelperMac : public ICrashDebugHelper
 {
 public:
@@ -18,7 +20,6 @@ public:
 	 */
 	virtual bool ParseCrashDump(const FString& InCrashDumpName, FCrashDebugInfo& OutCrashDebugInfo) override;
 
-	virtual bool SyncAndDebugCrashDump(const FString& InCrashDumpName) override;
 
 	/**
 	 *	Parse the given crash dump, and generate a report. 
diff --git a/Engine/Source/Developer/CrashDebugHelper/Private/Windows/CrashDebugHelperWindows.cpp b/Engine/Source/Developer/CrashDebugHelper/Private/Windows/CrashDebugHelperWindows.cpp
--- a/Engine/Source/Developer/CrashDebugHelper/Private/Windows/CrashDebugHelperWindows.cpp
+++ b/Engine/Source/Developer/CrashDebugHelper/Private/Windows/CrashDebugHelperWindows.cpp
@@ -1,9 +1,10 @@
 // Copyright 1998-2015 Epic Games, Inc. All Rights Reserved.
 
 #include "CrashDebugHelperPrivatePCH.h"
-#include "Database.h"
-
+#include "CrashDebugHelperWindows.h"
+#include "WindowsPlatformStackWalkExt.h"
 
+#include "EngineVersion.h"
 #include "ISourceControlModule.h"
 
 #include "AllowWindowsPlatformTypes.h"
@@ -17,410 +18,109 @@ FCrashDebugHelperWindows::~FCrashDebugHelperWindows()
 {
 }
 
-//
-//	Helper functions for processing windows crash dumps
-//
 
-/**
- *	Verify the given crash dump is Windows based
- *
- *	@param	InMiniDumpHeader		The header from the crash dump
- *
- *	@return	bool					true if it is a valid crash dump, false if not
- */
-bool VerifyCrashDump_Windows(MINIDUMP_HEADER& InMiniDumpHeader)
+bool FCrashDebugHelperWindows::ParseCrashDump(const FString& InCrashDumpName, FCrashDebugInfo& OutCrashDebugInfo)
 {
-	// See if the signature is "MDMP"
-	if( InMiniDumpHeader.Signature == 0x504d444d )
-	{
-		return true;
-	}
-
-	return false;
+	return true;
 }
 
-/**
- *	Extract the UE engine version from the crash dump
- *
- *	@param	InFileHandle		The handle to the opened crash dump file.
- *	@param	InMiniDumpHeader	The header from the crash dump
- *	@param	OutEngineVersion	The engine version if successfully extracted
- *
- *	@return	bool				true if successfully extracted, false if not
- */
-bool ExtractEngineVersion_Windows(HANDLE InFileHandle, MINIDUMP_HEADER& InMiniDumpHeader, int32& OutEngineVersion, FString& OutPlatform)
+bool FCrashDebugHelperWindows::CreateMinidumpDiagnosticReport( const FString& InCrashDumpFilename )
 {
-	bool bResult = false;
-
-	// Read the directories
-	MINIDUMP_DIRECTORY ModuleListDirectory;
-	FMemory::Memzero(&ModuleListDirectory, sizeof(MINIDUMP_DIRECTORY));
-	int32 ModuleListStreamIdx = -1;
+	const bool bSyncSymbols = FParse::Param( FCommandLine::Get(), TEXT( "SyncSymbols" ) );
+	const bool bAnnotate = FParse::Param( FCommandLine::Get(), TEXT( "Annotate" ) );
+	const bool bUseSCC = bSyncSymbols || bAnnotate;
 
-	bool bFailed = false;
-	LONG StreamOffset = InMiniDumpHeader.StreamDirectoryRva;
-	for (ULONG StreamIdx = 0; (StreamIdx < InMiniDumpHeader.NumberOfStreams) && !bFailed; StreamIdx++)
+	if( bUseSCC )
 	{
-		SetFilePointer(InFileHandle, StreamOffset, NULL, FILE_BEGIN);
-
-		MINIDUMP_DIRECTORY MDDirectory;
-		DWORD NumByteToRead = sizeof(MINIDUMP_DIRECTORY);
-		DWORD NumBytesRead = 0;
-		if (ReadFile(InFileHandle, &MDDirectory, NumByteToRead, &NumBytesRead, NULL) == TRUE)
-		{
-			switch (MDDirectory.StreamType)
-			{
-			// Reserved. Do not use this enumeration value.
-			case UnusedStream:
-			case ReservedStream0:
-			case ReservedStream1:
-				break;
-
-			// * The stream contains thread information and stack. For more information, see MINIDUMP_THREAD_LIST.
-			case ThreadListStream:
-				break;
-
-			// * The stream contains module information with the base, size and version info. For more information, see MINIDUMP_MODULE_LIST.
-			case ModuleListStream:
-				ModuleListStreamIdx = StreamIdx;
-				FMemory::Memcpy(&ModuleListDirectory, &MDDirectory, sizeof(MINIDUMP_DIRECTORY));
-				break;
-
-			// The stream contains memory allocation information. For more information, see MINIDUMP_MEMORY_LIST.
-			case MemoryListStream:
-				break;
-
-			// * The stream contains exception information with the code. For more information, see MINIDUMP_EXCEPTION_STREAM.
-			case ExceptionStream:
-				break;
-
-			// * The stream contains general system information. For more information, see MINIDUMP_SYSTEM_INFO.
-			case SystemInfoStream:
-				break;
-
-			// The stream contains extended thread information. For more information, see MINIDUMP_THREAD_EX_LIST.
-			case ThreadExListStream:
-				break;
-
-			// The stream contains memory allocation information. For more information, see MINIDUMP_MEMORY64_LIST.
-			case Memory64ListStream:
-				break;
-
-			// The stream contains an ANSI string used for documentation purposes.
-			case CommentStreamA:
-				break;
-
-			// The stream contains a Unicode string used for documentation purposes.
-			case CommentStreamW:
-				break;
-
-			// The stream contains high-level information about the active operating system handles. For more information, see MINIDUMP_HANDLE_DATA_STREAM.
-			case HandleDataStream:
-				break;
-
-			// The stream contains function table information. For more information, see MINIDUMP_FUNCTION_TABLE_STREAM.
-			case FunctionTableStream:
-				break;
-
-			// The stream contains module information for the unloaded modules. For more information, see MINIDUMP_UNLOADED_MODULE_LIST.
-			case UnloadedModuleListStream:
-				break;
-
-			// The stream contains miscellaneous information. For more information, see MINIDUMP_MISC_INFO or MINIDUMP_MISC_INFO_2.
-			case MiscInfoStream:
-				break;
-
-			// The stream contains memory region description information. It corresponds to the information that would be returned for the process from the 
-			// VirtualQuery function. For more information, see MINIDUMP_MEMORY_INFO_LIST.
-			case MemoryInfoListStream:
-				break;
-
-			// * The stream contains thread state information. For more information, see MINIDUMP_THREAD_INFO_LIST.
-			case ThreadInfoListStream:
-				break;
-
-			// This stream contains operation list information. For more information, see MINIDUMP_HANDLE_OPERATION_LIST.
-			case HandleOperationListStream:
-				break;
-
-			// Any value greater than this value will not be used by the system and can be used to represent application-defined data streams. For more information, see MINIDUMP_USER_STREAM.
-			case LastReservedStream:
-			default:
-				break;
-			}
-		}
-		else
-		{
-			bFailed = true;
-		}
-
-		StreamOffset += sizeof(MINIDUMP_DIRECTORY);
+		InitSourceControl( false );
 	}
 
-	if (ModuleListStreamIdx != -1)
-	{
-		// We found the module list stream. Parse that...
-		// Prep the file pointer to the right location
-		SetFilePointer(InFileHandle, ModuleListDirectory.Location.Rva, NULL, FILE_BEGIN);
+	FWindowsPlatformStackWalkExt WindowsStackWalkExt( CrashInfo );
 
-		ULONG32 NumberOfModules = 0;
-		DWORD NumByteToRead = sizeof(ULONG32);
-		DWORD NumBytesRead = 0;
+	const bool bReady = WindowsStackWalkExt.InitStackWalking();
 
-		if (ReadFile(InFileHandle, &NumberOfModules, NumByteToRead, &NumBytesRead, NULL) == TRUE)
+	bool bHasAtLeastThreeValidFunctions = false;
+	if( bReady && WindowsStackWalkExt.OpenDumpFile( InCrashDumpFilename ) )
+	{
+		if( CrashInfo.BuiltFromCL != FCrashInfo::INVALID_CHANGELIST )
 		{
-			//@todo. This is excess as we really only care about the executable module... clean this up.
-			TMap<FString, MINIDUMP_MODULE> MDModulesMap;
-			TArray<MINIDUMP_MODULE> MDModules;
-			MDModules.Empty(NumberOfModules);
-			MDModules.AddZeroed(NumberOfModules);
-
-			bool bFailed = false;
-			for (uint32 ModuleIdx = 0; (ModuleIdx < NumberOfModules) && !bFailed; ModuleIdx++)
+			// CrashInfo now contains a changelist to lookup a label for
+			if( bSyncSymbols )
 			{
-				NumByteToRead = sizeof(MINIDUMP_MODULE);
-				NumBytesRead = 0;
-				if (ReadFile(InFileHandle, &(MDModules[ModuleIdx]), NumByteToRead, &NumBytesRead, NULL) == FALSE)
-				{
-					bFailed = true;
-				}
-			}
+				// Get the build version from the crash info.
+				FCrashModuleInfo ExeFileVersion;
+				WindowsStackWalkExt.GetExeFileVersionAndModuleList( ExeFileVersion );
 
-			if (bFailed == false)
-			{
-				for (uint32 NameIdx = 0; (NameIdx < (uint32)MDModules.Num()) && !bFailed; NameIdx++)
-				{
-					MINIDUMP_MODULE& Module = MDModules[NameIdx];
-					SetFilePointer(InFileHandle, Module.ModuleNameRva, NULL, FILE_BEGIN);
+				FindSymbolsAndBinariesStorage();
 
-					WCHAR TempBuffer[8192];
-					ULONG32 NameLength;
-					NumByteToRead = sizeof(ULONG32);
-					NumBytesRead = 0;
-					if (ReadFile(InFileHandle, &NameLength, NumByteToRead, &NumBytesRead, NULL) == TRUE)
-					{
-						if ((NameLength / 2) < 8192)
-						{
-							NumByteToRead = NameLength;
-							NumBytesRead = 0;
-							if (ReadFile(InFileHandle, &TempBuffer, NumByteToRead, &NumBytesRead, NULL) == TRUE)
-							{
-								TempBuffer[NameLength/2] = 0;
-								FString ModuleName = TempBuffer;
-								MDModulesMap.Add(ModuleName, Module);
-							}
-							else
-							{
-								bFailed = true;
-							}
-						}
-						else
-						{
-							bFailed = true;
-						}
-					}
-					else
-					{
-						bFailed = true;
-					}
+				const bool bSynced = SyncModules();
+				// Without symbols we can't decode the provided minidump.
+				if( !bSynced )
+				{
+					return 0;
 				}
 			}
 
-			if (bFailed == false)
-			{
-				// Find the executable module...
-				for (TMap<FString,MINIDUMP_MODULE>::TIterator DumpIt(MDModulesMap); DumpIt; ++DumpIt)
-				{
-					FString& ModuleName = DumpIt.Key();
-					MINIDUMP_MODULE& Module = DumpIt.Value();
-
-					//@todo. We need to be able to determine other UE4 executables like SlateViewer...
-					//@todo rocket Game suffix! We need another way to discover game executables
-
-					FString GameExeName = TEXT("Game.exe");
-
-					FString Win64GameDirName(TEXT("Game\\Binaries\\Win64\\"));
-					FString Win64UE4ExeName(TEXT("Binaries\\Win64\\UE4"));
-
-					FString Win32GameDirName(TEXT("Game\\Binaries\\Win32\\"));
-					FString Win32UE4ExeName(TEXT("Binaries\\Win32\\UE4"));
-					FString UE4ExeName = TEXT(".exe");
-
-					bool bIsOurExecutable = false;
-					if ((ModuleName.Contains(Win64GameDirName) && ModuleName.EndsWith(GameExeName)) || 
-						(ModuleName.Contains(Win64UE4ExeName) && ModuleName.EndsWith(UE4ExeName)))
-					{
-						bIsOurExecutable = true;
-						OutPlatform = TEXT("Win64");
-					}
-
-					if ((ModuleName.Contains(Win32GameDirName) && ModuleName.EndsWith(GameExeName)) ||
-						(ModuleName.Contains(Win32UE4ExeName) && ModuleName.EndsWith(UE4ExeName)))
-					{
-						bIsOurExecutable = true;
-						OutPlatform = TEXT("Win32");
-					}
-
-					if (bIsOurExecutable == true)
-					{
-						int32 VerA = (Module.VersionInfo.dwProductVersionMS & 0xFFFF0000) >> 16;
-						int32 VerB = (Module.VersionInfo.dwProductVersionMS & 0x0000FFFF);
-						int32 VerC = (Module.VersionInfo.dwProductVersionLS & 0xFFFF0000) >> 16;
-						int32 VerD = (Module.VersionInfo.dwProductVersionLS & 0x0000FFFF);
+			// Initialise the symbol options
+			WindowsStackWalkExt.InitSymbols();
 
-						if (VerD == 0)
-						{
-							// Format is #.#.ENGINEVERSION.#
-							OutEngineVersion = VerC;
-						}
-						else
-						{
-							// Format is #.#.HIChangeList.LOWChangeList
-							OutEngineVersion = (VerC << 16) | VerD;
-						}
+			// Set the symbol path based on the loaded modules
+			WindowsStackWalkExt.SetSymbolPathsFromModules();
 
-						bResult = true;
-					}
-				}
-			}
-		}
-		else
-		{
-			UE_LOG(LogCrashDebugHelper, Warning, TEXT("ParseMiniDumpFile: Failed to read module list"));
-		}
-	}
-	else
-	{
-		UE_LOG(LogCrashDebugHelper, Warning, TEXT("Failed to find module list stream in crash dump..."));
-	}
+			// Get all the info we should ever need about the modules
+			WindowsStackWalkExt.GetModuleInfoDetailed();
 
-	return bResult;
-}
+			// Get info about the system that created the minidump
+			WindowsStackWalkExt.GetSystemInfo();
 
-bool FCrashDebugHelperWindows::LaunchDebugger(const FString& InCrashDumpFilename, FCrashDebugInfo& OutCrashDebugInfo)
-{
-	// Ensure we are in a valid state to launch the debugger
-	return false;
-}
+			// Get all the thread info
+			WindowsStackWalkExt.GetThreadInfo();
 
-bool FCrashDebugHelperWindows::ParseCrashDump_Windows(const TCHAR* InCrashDumpFilename, FCrashDebugInfo& OutCrashDebugInfo)
-{
-	bool bResult = false;
+			// Get exception info
+			WindowsStackWalkExt.GetExceptionInfo();
 
-	HANDLE FileHandle = CreateFileW(InCrashDumpFilename, GENERIC_READ, FILE_SHARE_READ, NULL, OPEN_EXISTING, 0, NULL);
-	if (FileHandle != INVALID_HANDLE_VALUE)
-	{
-		// Read the header
-		MINIDUMP_HEADER	MDHeader;
-		FMemory::Memzero(&MDHeader, sizeof(MINIDUMP_HEADER));
-		DWORD NumByteToRead = sizeof(MINIDUMP_HEADER);
-		DWORD NumBytesRead;
-		if (ReadFile(FileHandle, &MDHeader, NumByteToRead, &NumBytesRead, NULL) == TRUE)
-		{
-			UE_LOG(LogCrashDebugHelper, Display, TEXT("ParseCrashDump_Windows: Read header from %s"), InCrashDumpFilename);
+			// Get the callstacks for each thread
+			bHasAtLeastThreeValidFunctions = WindowsStackWalkExt.GetCallstacks() >= 3;
 
-			// Verify it is a legitimate crash dump
-			if (VerifyCrashDump_Windows(MDHeader) == true)
+			// Sync the source file where the crash occurred
+			if( CrashInfo.SourceFile.Len() > 0 )
 			{
-				UE_LOG(LogCrashDebugHelper, Display, TEXT("ParseCrashDump_Windows: Verified windows crash dump: %s"), InCrashDumpFilename);
-
-				// Extract the engine version
-				if (ExtractEngineVersion_Windows(FileHandle, MDHeader, OutCrashDebugInfo.EngineVersion, OutCrashDebugInfo.PlatformName) == true)
+				if( bSyncSymbols && CrashInfo.BuiltFromCL > 0 )
 				{
-					UE_LOG(LogCrashDebugHelper, Display, TEXT("ParseCrashDump_Windows: Extracted engine version from crash dump: %s"), InCrashDumpFilename);
-					UE_LOG(LogCrashDebugHelper, Display, TEXT("\tEngine Version = %d"), OutCrashDebugInfo.EngineVersion);
-					UE_LOG(LogCrashDebugHelper, Display, TEXT("\tPlatform       = %s"), *(OutCrashDebugInfo.PlatformName));
+					UE_LOG( LogCrashDebugHelper, Log, TEXT( "Using CL %i to sync crash source file" ), CrashInfo.BuiltFromCL );
+					SyncSourceFile();
+				}
 
-					// Obsolete?
-					//OutCrashDebugInfo.SourceControlLabel = GetLabelFromEngineVersion(OutCrashDebugInfo.EngineVersion);
-					if (OutCrashDebugInfo.SourceControlLabel.Len() > 0)
-					{
-						
-						UE_LOG(LogCrashDebugHelper, Display, TEXT("ParseCrashDump_Windows: Retrieve changelist label: %s"), *(OutCrashDebugInfo.SourceControlLabel));
-						bResult = true;
-					}
-					else
-					{
-						UE_LOG(LogCrashDebugHelper, Warning, TEXT("ParseCrashDump_Windows: Failed to retrieve changelist label for engine version: %d"), OutCrashDebugInfo.EngineVersion);
-					}
+				// Try to annotate the file if requested
+				bool bAnnotationSuccessful = false;
+				if( bAnnotate )
+				{
+					bAnnotationSuccessful = AddAnnotatedSourceToReport();
 				}
-				else
+
+				// If annotation is not requested, or failed, add the standard source context
+				if( !bAnnotationSuccessful )
 				{
-					UE_LOG(LogCrashDebugHelper, Warning, TEXT("ParseCrashDump_Windows: Failed to extract engine version from crash dump: %s"), InCrashDumpFilename);
+					AddSourceToReport();
 				}
 			}
-			else
-			{
-				UE_LOG(LogCrashDebugHelper, Warning, TEXT("ParseCrashDump_Windows: Invalid windows crash dump: %s"), InCrashDumpFilename);
-			}
-
-			CloseHandle(FileHandle);
 		}
 		else
 		{
-			UE_LOG(LogCrashDebugHelper, Warning, TEXT("ParseCrashDump_Windows: Failed to read header from %s"), InCrashDumpFilename);
+			UE_LOG( LogCrashDebugHelper, Warning, TEXT( "Invalid built from changelist" ) );
 		}
 	}
 	else
 	{
-		UE_LOG(LogCrashDebugHelper, Warning, TEXT("ParseCrashDump_Windows: Failed to open crash dump file: %s"), InCrashDumpFilename);
-	}
-	return bResult;
-}
-
-
-bool FCrashDebugHelperWindows::ParseCrashDump(const FString& InCrashDumpName, FCrashDebugInfo& OutCrashDebugInfo)
-{
-	if (bInitialized == false)
-	{
-		UE_LOG(LogCrashDebugHelper, Warning, TEXT("ParseCrashDump: CrashDebugHelper not initialized"));
-		return false;
+		UE_LOG( LogCrashDebugHelper, Warning, TEXT( "Failed to open crash dump file: %s" ), *InCrashDumpFilename );
 	}
 
-	if (ParseCrashDump_Windows(*InCrashDumpName, OutCrashDebugInfo) == true)
-	{
-		return true;
-	}
-
-	return false;
-}
-
-bool FCrashDebugHelperWindows::SyncAndDebugCrashDump(const FString& InCrashDumpName)
-{
-	if (bInitialized == false)
-	{
-		UE_LOG(LogCrashDebugHelper, Warning, TEXT("SyncAndDebugCrashDump: CrashDebugHelper not initialized"));
-		return false;
-	}
-
-	FCrashDebugInfo CrashDebugInfo;
-	
-	// Parse the crash dump if it hasn't already been...
-	if (ParseCrashDump(InCrashDumpName, CrashDebugInfo) == false)
-	{
-		UE_LOG(LogCrashDebugHelper, Warning, TEXT("SyncAndDebugCrashDump: Failed to parse crash dump %s"), *InCrashDumpName);
-		return false;
-	}
-
-	// Now sync the required pdb files
-	if (SyncRequiredFilesForDebuggingFromLabel(CrashDebugInfo.SourceControlLabel, CrashDebugInfo.PlatformName) == true)
-	{
-		// Launch the debugger
-		if (LaunchDebugger(InCrashDumpName, CrashDebugInfo) == true)
-		{
-			return true;
-		}
-		else
-		{
-			UE_LOG(LogCrashDebugHelper, Warning, TEXT("SyncAndDebugCrashDump: Failed to launch debugger for crash dump %s"), *InCrashDumpName);
-		}
-	}
-	else
+	if( bUseSCC )
 	{
-		UE_LOG(LogCrashDebugHelper, Warning, TEXT("SyncAndDebugCrashDump: Failed to sync required files for debugging crash dump %s"), *InCrashDumpName);
+		ShutdownSourceControl();
 	}
 
-	return false;
+	return bHasAtLeastThreeValidFunctions;
 }
 
 #include "HideWindowsPlatformTypes.h"
diff --git a/Engine/Source/Developer/CrashDebugHelper/Private/Windows/CrashDebugHelperWindows.h b/Engine/Source/Developer/CrashDebugHelper/Private/Windows/CrashDebugHelperWindows.h
--- a/Engine/Source/Developer/CrashDebugHelper/Private/Windows/CrashDebugHelperWindows.h
+++ b/Engine/Source/Developer/CrashDebugHelper/Private/Windows/CrashDebugHelperWindows.h
@@ -2,6 +2,8 @@
 
 #pragma once
 
+#include "CrashDebugHelper.h"
+
 class FCrashDebugHelperWindows : public ICrashDebugHelper
 {
 public:
@@ -18,8 +20,6 @@ public:
 	 */
 	virtual bool ParseCrashDump(const FString& InCrashDumpName, FCrashDebugInfo& OutCrashDebugInfo) override;
 
-	virtual bool SyncAndDebugCrashDump(const FString& InCrashDumpName) override;
-
 	/**
 	 *	Parse the given crash dump, and generate a report. 
 	 *
@@ -28,24 +28,6 @@ public:
 	 *	@return	bool				true if successful, false if not
 	 */
 	virtual bool CreateMinidumpDiagnosticReport( const FString& InCrashDumpName ) override;
-protected:
-	/**
-	 *	Process the given crash dump file
-	 *
-	 *	@param	InCrashDumpFilename		The name of the crash dump file
-	 *	@param	OutCrashDebugInfo		The info on the crash dump
-	 *
-	 *	@return	bool					true if successful, false if not
-	 */
-	bool ParseCrashDump_Windows(const TCHAR* InCrashDumpFilename, FCrashDebugInfo& OutCrashDebugInfo);
-
-	/**
-	 *	Launch the debugger for the current crash dump
-	 *
-	 *	@return bool		true if successful, false if not
-	 */
-	bool LaunchDebugger(const FString& InCrashDumpFilename, FCrashDebugInfo& OutCrashDebugInfo);
-
 };
 
 typedef FCrashDebugHelperWindows FCrashDebugHelper;
diff --git a/Engine/Source/Developer/CrashDebugHelper/Private/Windows/MinidumpParser.cpp b/Engine/Source/Developer/CrashDebugHelper/Private/Windows/MinidumpParser.cpp
deleted file mode 100644
diff --git a/Engine/Source/Developer/CrashDebugHelper/Private/Windows/WindowsPlatformStackWalkExt.cpp b/Engine/Source/Developer/CrashDebugHelper/Private/Windows/WindowsPlatformStackWalkExt.cpp
--- a/Engine/Source/Developer/CrashDebugHelper/Private/Windows/WindowsPlatformStackWalkExt.cpp
+++ b/Engine/Source/Developer/CrashDebugHelper/Private/Windows/WindowsPlatformStackWalkExt.cpp
@@ -1,9 +1,9 @@
 // Copyright 1998-2015 Epic Games, Inc. All Rights Reserved.
 
 #include "CrashDebugHelperPrivatePCH.h"
-
-#include "ModuleManager.h"
+#include "GenericPlatform/GenericPlatformContext.h"
 #include "WindowsPlatformStackWalkExt.h"
+#include "CrashDebugPDBCache.h"
 
 #include "AllowWindowsPlatformTypes.h"
 #include "dbgeng.h"
@@ -27,9 +27,7 @@ FWindowsPlatformStackWalkExt::~FWindowsPlatformStackWalkExt()
 }
 
 
-/**
- * Initialise the COM interface to grab stacks
- */
+
 bool FWindowsPlatformStackWalkExt::InitStackWalking()
 {
 	if (!FWindowsPlatformMisc::CoInitialize())
@@ -46,9 +44,6 @@ bool FWindowsPlatformStackWalkExt::InitStackWalking()
 	return true;
 }
 
-/**
- * Shutdown COM
- */
 void FWindowsPlatformStackWalkExt::ShutdownStackWalking()
 {
 	Advanced->Release();
@@ -61,7 +56,7 @@ void FWindowsPlatformStackWalkExt::ShutdownStackWalking()
 }
 
 /**
- * Set the options we want for symbol lookup
+ * 
  */
 void FWindowsPlatformStackWalkExt::InitSymbols()
 {
@@ -79,15 +74,14 @@ void FWindowsPlatformStackWalkExt::InitSymbols()
 	// This option allows for undecorated names to be handled by the symbol engine.
 	SymOpts |= SYMOPT_UNDNAME;
 
+#if	_DEBUG
 	// Disable by default as it can be very spammy/slow.  Turn it on if you are debugging symbol look-up!
 	SymOpts |= SYMOPT_DEBUG;
+#endif // _DEBUG
 
 	Symbol->SetSymbolOptions( SymOpts );
 }
 
-/** 
- * Grab the branch relative path of the binary
- */
 FString FWindowsPlatformStackWalkExt::ExtractRelativePath( const TCHAR* BaseName, TCHAR* FullName )
 {
 	FString FullPath = FString( FullName ).ToLower();
@@ -126,7 +120,7 @@ void FWindowsPlatformStackWalkExt::GetExeFileVersionAndModuleList( FCrashModuleI
 	ULONG UnloadedModuleCount = 0;
 	Symbol->GetNumberModules( &LoadedModuleCount, &UnloadedModuleCount );
 
-	CrashInfo.Log( FString::Printf( TEXT( " ... found %d loaded, and %d unloaded modules." ), LoadedModuleCount, UnloadedModuleCount ) );
+	UE_LOG( LogCrashDebugHelper, Log, TEXT( "Modules loaded: %i, unloaded: %i" ), LoadedModuleCount, UnloadedModuleCount );
 
 	// Find the relative names of all the modules so we know which files to sync
 	int ExecutableIndex = -1;
@@ -164,7 +158,7 @@ void FWindowsPlatformStackWalkExt::GetExeFileVersionAndModuleList( FCrashModuleI
 	}
 	else
 	{
-		CrashInfo.Log( FString::Printf( TEXT( " ... unable to locate executable." ) ) );
+		UE_LOG( LogCrashDebugHelper, Warning, TEXT( "Unable to locate the executable" ) );
 	}
 }
 
@@ -262,9 +256,6 @@ public:
 	}
 };
 
-/** 
- * Set the symbol paths based on the module paths
- */
 void FWindowsPlatformStackWalkExt::SetSymbolPathsFromModules()
 {
 	const bool bUseCachedData = CrashInfo.PDBCacheEntry.IsValid();
@@ -359,29 +350,45 @@ void FWindowsPlatformStackWalkExt::SetSymbolPathsFromModules()
 		FString BinariesDir = FString(FPlatformProcess::BaseDir());
 		if ( !FPaths::FileExists( FPaths::Combine(*BinariesDir, TEXT("symsrv.dll")) ) )
 		{
-			CrashInfo.Log( FString::Printf(TEXT("Error: symsrv.dll was not detected in '%s'. Microsoft symbols will not be downloaded!"), *BinariesDir) );
+			UE_LOG( LogCrashDebugHelper, Warning, TEXT( "Error: symsrv.dll was not detected in: %s. Microsoft symbols will not be downloaded!" ), *BinariesDir );
 		}
 		if ( !FPaths::FileExists( FPaths::Combine(*BinariesDir, TEXT("symsrv.yes")) ) )
 		{
-			CrashInfo.Log( FString::Printf(TEXT("symsrv.yes was not detected in '%s'. This will cause a popup to confirm the license."), *BinariesDir) );
+			UE_LOG( LogCrashDebugHelper, Warning, TEXT( "symsrv.yes was not detected in: %s. This will cause a popup to confirm the license." ), *BinariesDir );
 		}
 		if ( !FPaths::FileExists( FPaths::Combine(*BinariesDir, TEXT("dbghelp.dll")) ) )
 		{
-			CrashInfo.Log( FString::Printf(TEXT("Error: dbghelp.dll was not detected in '%s'. Microsoft symbols will not be downloaded!"), *BinariesDir) );
+			UE_LOG( LogCrashDebugHelper, Warning, TEXT( "Error: dbghelp.dll was not detected in: %s. Microsoft symbols will not be downloaded!" ), *BinariesDir );
 		}
 
 		Symbol->AppendImagePathWide( TEXT( "SRV*..\\..\\Intermediate\\SymbolCache*http://msdl.microsoft.com/download/symbols" ) );
 		Symbol->AppendSymbolPathWide( TEXT( "SRV*..\\..\\Intermediate\\SymbolCache*http://msdl.microsoft.com/download/symbols" ) );
 	}
 
-	TCHAR SymbolPath[1024] = { 0 };
-	Symbol->GetSymbolPathWide( SymbolPath, 1024, NULL );
-	CrashInfo.Log( FString::Printf( TEXT( " ... setting symbol path to: '%s'" ), SymbolPath ) );
+	TCHAR SymbolPath[16384] = { 0 };
+
+	Symbol->GetSymbolPathWide( SymbolPath, ARRAY_COUNT(SymbolPath), NULL );
+	TArray<FString> SymbolPaths;
+	FString( SymbolPath ).ParseIntoArray(&SymbolPaths, TEXT(";"), true );
+
+	UE_LOG( LogCrashDebugHelper, Log, TEXT( "Symbol paths" ) );
+	for( const auto& It : SymbolPaths )
+	{
+		UE_LOG( LogCrashDebugHelper, Log, TEXT( "    %s" ), *It );
+	}
 
-	Symbol->GetImagePathWide( SymbolPath, 1024, NULL );
-	CrashInfo.Log( FString::Printf( TEXT( " ... setting image path to: '%s'" ), SymbolPath ) );
+	Symbol->GetImagePathWide( SymbolPath, ARRAY_COUNT( SymbolPath ), NULL );
+	TArray<FString> ImagePaths;
+	FString( SymbolPath ).ParseIntoArray( &ImagePaths, TEXT( ";" ), true );
+	
+	UE_LOG( LogCrashDebugHelper, Log, TEXT( "Image paths" ) );
+	for( const auto& It : ImagePaths )
+	{
+		UE_LOG( LogCrashDebugHelper, Log, TEXT( "    %s" ), *It );
+	}
 }
 
+/** Helper class used to sort modules by name. */
 class FSortModulesByName
 {
 public:
@@ -403,9 +410,6 @@ public:
 	}
 };
 
-/** 
- * Get detailed info about each module
- */
 void FWindowsPlatformStackWalkExt::GetModuleInfoDetailed()
 {
 	// The the number of loaded modules
@@ -455,10 +459,6 @@ void FWindowsPlatformStackWalkExt::GetModuleInfoDetailed()
 
 	CrashInfo.Modules.Sort( FSortModulesByName() );
 }
-
-/**
- * Check to see if the stack address resides within one of the loaded modules i.e. is it valid?
- */
 bool FWindowsPlatformStackWalkExt::IsOffsetWithinModules( uint64 Offset )
 {
 	for( int ModuleIndex = 0; ModuleIndex < CrashInfo.Modules.Num(); ModuleIndex++ )
@@ -473,9 +473,6 @@ bool FWindowsPlatformStackWalkExt::IsOffsetWithinModules( uint64 Offset )
 	return false;
 }
 
-/** 
- * Extract the system info of the crash from the minidump
- */
 void FWindowsPlatformStackWalkExt::GetSystemInfo()
 {
 	FCrashSystemInfo& SystemInfo = CrashInfo.SystemInfo;
@@ -521,16 +518,6 @@ void FWindowsPlatformStackWalkExt::GetSystemInfo()
 	SystemInfo.ProcessorCount = ProcessorCount;
 }
 
-/** 
- * Extract the thread info from the minidump
- */
-void FWindowsPlatformStackWalkExt::GetThreadInfo()
-{
-}
-
-/** 
- * Extract info about the exception that caused the crash
- */
 void FWindowsPlatformStackWalkExt::GetExceptionInfo()
 {
 	FCrashExceptionInfo& Exception = CrashInfo.Exception;
@@ -547,146 +534,147 @@ void FWindowsPlatformStackWalkExt::GetExceptionInfo()
 	Exception.ExceptionString = Description;
 }
 
-/** 
- * Get the callstack of the crash. Returns true if at least one function name was found in the debugging symbols
- */
-bool FWindowsPlatformStackWalkExt::GetCallstacks()
+
+int FWindowsPlatformStackWalkExt::GetCallstacks()
 {
-	bool bAtLeastOneFunctionNameFound = false;
+	const int32 MAX_NAME_LENGHT = FProgramCounterSymbolInfo::MAX_NAME_LENGHT;
+	int32 NumValidFunctionNames = 0;
 
 	FCrashExceptionInfo& Exception = CrashInfo.Exception;
 
 	FMemMark Mark(FMemStack::Get());
 
 	//const float int int32 FString
-	const int32 ConextSize = 4096;
-	byte* Context = new(FMemStack::Get()) byte[ConextSize];
+	const int32 ContextSize = 4096;
+	byte* Context = new(FMemStack::Get()) byte[ContextSize];
 	ULONG DebugEvent = 0;
 	ULONG ProcessID = 0;
 	ULONG ThreadID = 0;
-	ULONG ContextSize = 0;
+	ULONG ContextUsed = 0;
 
 	// Get the context of the crashed thread
-	if( Control->GetStoredEventInformation( &DebugEvent, &ProcessID, &ThreadID, Context, ConextSize, &ContextSize, NULL, 0, 0) )
+	if( Control->GetStoredEventInformation( &DebugEvent, &ProcessID, &ThreadID, Context, ContextSize, &ContextUsed, NULL, 0, 0) )
 	{
-		return bAtLeastOneFunctionNameFound;
+		return NumValidFunctionNames;
 	}
 
 	// Some magic number checks
-	if( ContextSize == 716 )
+	if( ContextUsed == 716 )
 	{
-		CrashInfo.Log( FString::Printf( TEXT( " ... context size matches x86 sizeof( CONTEXT )" ) ) );
+		UE_LOG( LogCrashDebugHelper, Log, TEXT( "Context size matches x86 sizeof( CONTEXT )" ) );
 	}
-	else if( ContextSize == 1232 )
+	else if( ContextUsed == 1232 )
 	{
-		CrashInfo.Log( FString::Printf( TEXT( " ... context size matches x64 sizeof( CONTEXT )" ) ) );
+		UE_LOG( LogCrashDebugHelper, Log, TEXT( "Context size matches x64 sizeof( CONTEXT )" ) );
 	}
 
 	// Get the entire stack trace
 	const uint32 MaxFrames = 128;
-	const uint32 MaxFramesSize = MaxFrames * ContextSize;
+	const uint32 MaxFramesSize = MaxFrames * ContextUsed;
 
 	DEBUG_STACK_FRAME* StackFrames = new(FMemStack::Get()) DEBUG_STACK_FRAME[MaxFrames];
 	ULONG Count = 0;
 	bool bFoundSourceFile = false;
 	void* ContextData = FMemStack::Get().PushBytes( MaxFramesSize, 0 );
 	FMemory::Memzero( ContextData, MaxFramesSize );
-	HRESULT HR = Control->GetContextStackTrace( Context, ContextSize, StackFrames, MaxFrames, ContextData, MaxFramesSize, ContextSize, &Count );
+	HRESULT HR = Control->GetContextStackTrace( Context, ContextUsed, StackFrames, MaxFrames, ContextData, MaxFramesSize, ContextUsed, &Count );
 
 	for( uint32 StackIndex = 0; StackIndex < Count; StackIndex++ )
-	{
-		FString Line;
-		uint64 Offset = StackFrames[StackIndex].InstructionOffset;
+	{	
+		const uint64 Offset = StackFrames[StackIndex].InstructionOffset;
 
 		if( IsOffsetWithinModules( Offset ) )
 		{
 			// Get the module, function, and offset
 			uint64 Displacement = 0;
-			TCHAR FunctionName[MAX_PATH] = { 0 };
+			TCHAR NameByOffset[MAX_PATH] = {0};
+			Symbol->GetNameByOffsetWide( Offset, NameByOffset, ARRAYSIZE( NameByOffset ) - 1, NULL, &Displacement );
+			FString ModuleAndFunction = NameByOffset;
 
-			Symbol->GetNameByOffsetWide( Offset, FunctionName, ARRAYSIZE( FunctionName ) - 1, NULL, &Displacement );
+			// Don't care about any more entries higher than this
+			if( ModuleAndFunction.Contains( TEXT( "tmainCRTStartup" ) ) )
+			{
+				break;
+			}
 
 			// Look for source file name and line number
 			TCHAR SourceName[MAX_PATH] = { 0 };
 			ULONG LineNumber = 0;
-
 			Symbol->GetLineByOffsetWide( Offset, &LineNumber, SourceName, ARRAYSIZE( SourceName ) - 1, NULL, NULL );
 
-			Line = FunctionName;
-
 			// Remember the top of the stack to locate in the source file
-			if( !bFoundSourceFile && FString( SourceName ).Len() > 0 && LineNumber > 0 )
+			if( !bFoundSourceFile && FCString::Strlen( SourceName ) > 0 && LineNumber > 0 )
 			{
 				CrashInfo.SourceFile = ExtractRelativePath( TEXT( "source" ), SourceName );
 				CrashInfo.SourceLineNumber = LineNumber;
 				bFoundSourceFile = true;
 			}
 
-			// If we find an assert, the actual source file we're interested in is the next one up, so reset the source file found flag
-			if( Line.Contains( TEXT( "FDebug::AssertFailed" ), ESearchCase::CaseSensitive)
-				|| Line.Contains( TEXT("FDebug::EnsureNotFalse"), ESearchCase::CaseSensitive)
-				|| Line.Contains( TEXT( "FDebug::EnsureFailed" ), ESearchCase::CaseSensitive) )
+			FString ModuleName;
+			FString FunctionName;
+			// According to MSDN, the symbol name will include an ! if the function name could be discovered, delimiting it from the module name
+			// https://msdn.microsoft.com/en-us/library/windows/hardware/ff547186(v=vs.85).aspx
+			if( ModuleAndFunction.Contains( TEXT( "!" ) ) )
 			{
-				bFoundSourceFile = false;
+				NumValidFunctionNames++;
+ 
+				ModuleAndFunction.Split( TEXT( "!" ), &ModuleName, &FunctionName );
+				FunctionName += TEXT( "()" );
 			}
-
-			// According to MSDN, the symbol name will include an ! if the function name could be discovered, delimiting it from the module name
-			if( Line.Contains(TEXT( "!" )) )
+			else
 			{
-				Line += TEXT( "()" );
-				bAtLeastOneFunctionNameFound = true;
+				ModuleName = ModuleAndFunction;
 			}
 
-			Line += FString::Printf( TEXT( " + %d bytes" ), Displacement );
-			
-			if( LineNumber > 0 )
+			// @TODO yrx 2015-02-24 Add for other platforms
+			// If we find an assert, the actual source file we're interested in is the next one up, so reset the source file found flag
+			if( FunctionName.Len() > 0 )
 			{
-				Line += FString::Printf( TEXT( " [%s:%d]" ), SourceName, LineNumber );
+				if( FunctionName.Contains( TEXT( "FDebug::AssertFailed" ), ESearchCase::CaseSensitive )
+					|| FunctionName.Contains( TEXT( "FDebug::EnsureNotFalse" ), ESearchCase::CaseSensitive )
+					|| FunctionName.Contains( TEXT( "FDebug::EnsureFailed" ), ESearchCase::CaseSensitive ) )
+				{
+					bFoundSourceFile = false;
+				}
 			}
 
-			Exception.CallStackString.Add( Line );
-		}
-		else
-		{
-			break;
-		}
+			// FString InModuleName, FString InFunctionName, FString InFilename, uint32 InLineNumber, uint64 InSymbolDisplacement, uint64 InOffsetInModule, uint64 InProgramCounter
+			FProgramCounterSymbolInfoEx SymbolInfo( ModuleName, FunctionName, SourceName, LineNumber, Displacement, Offset, 0 );
+			FString GenericFormattedCallstackLine;
+			FGenericPlatformStackWalk::SymbolInfoToHumanReadableStringEx( SymbolInfo, GenericFormattedCallstackLine );
+			Exception.CallStackString.Add( GenericFormattedCallstackLine );
 
-		// Don't care about any more entries higher than this
-		if( Line.Contains(TEXT( "tmainCRTStartup" )) )
-		{
-			break;
+			UE_LOG( LogCrashDebugHelper, Log, TEXT( "%3u: %s" ), StackIndex, *GenericFormattedCallstackLine );
 		}
 	}
 
-	CrashInfo.Log( FString::Printf( TEXT( " ... callstack generated!" ) ) );
+	UE_LOG( LogCrashDebugHelper, Warning, TEXT( "Callstack generated with %i valid function names" ), NumValidFunctionNames );
 
-	return bAtLeastOneFunctionNameFound;
+	return NumValidFunctionNames;
 }
 
-/** 
- * Open a minidump as a new session
- */
+
 bool FWindowsPlatformStackWalkExt::OpenDumpFile( const FString& InCrashDumpFilename )
 {
 	const bool bFound = IFileManager::Get().FileSize( *InCrashDumpFilename ) != INDEX_NONE;
 	if( !bFound )
 	{
-		CrashInfo.Log( FString::Printf( TEXT( " ... failed to find minidump file: %s" ), *InCrashDumpFilename ) );
+		UE_LOG( LogCrashDebugHelper, Warning, TEXT( "Failed to find minidump file: %s" ), *InCrashDumpFilename );
 		return false;
 	}
 
 	if( Client->OpenDumpFileWide( *InCrashDumpFilename, NULL ) )
 	{
-		CrashInfo.Log( FString::Printf( TEXT( " ... failed to open minidump: %s" ), *InCrashDumpFilename ) );
+		UE_LOG( LogCrashDebugHelper, Warning, TEXT( "Failed to open minidump file: %s" ), *InCrashDumpFilename );
 		return false;
 	}
 
 	if( Control->WaitForEvent( 0, INFINITE ) != S_OK )
 	{
-		CrashInfo.Log( FString::Printf( TEXT( " ... failed while waiting for minidump to load: %s" ), *InCrashDumpFilename ) );
+		UE_LOG( LogCrashDebugHelper, Warning, TEXT( "Failed while waiting for minidump to load: %s" ), *InCrashDumpFilename );
 		return false;
 	}
 
-	CrashInfo.Log( FString::Printf( TEXT( "Successfully opened minidump: %s" ), *InCrashDumpFilename ) );
+	UE_LOG( LogCrashDebugHelper, Log, TEXT( "Successfully opened minidump: %s" ), *InCrashDumpFilename );
 	return true;
 }
diff --git a/Engine/Source/Developer/CrashDebugHelper/Private/Windows/WindowsPlatformStackWalkExt.h b/Engine/Source/Developer/CrashDebugHelper/Private/Windows/WindowsPlatformStackWalkExt.h
--- a/Engine/Source/Developer/CrashDebugHelper/Private/Windows/WindowsPlatformStackWalkExt.h
+++ b/Engine/Source/Developer/CrashDebugHelper/Private/Windows/WindowsPlatformStackWalkExt.h
@@ -19,22 +19,48 @@ struct FWindowsPlatformStackWalkExt
 	/** Destructor. */
 	~FWindowsPlatformStackWalkExt();
 
+	/** Initializes the COM interface to grab stacks. */
 	bool InitStackWalking();
+
+	/** Shutdowns COM. */
 	void ShutdownStackWalking();
 
-	void GetExeFileVersionAndModuleList( FCrashModuleInfo& out_ExeFileVersion );
+	/** Sets the options we want for symbol lookup. */
 	void InitSymbols();
-	bool OpenDumpFile( const FString& DumpFileName );
+
+	/** Grabs the branch relative path of the binary. */
+	static FString ExtractRelativePath( const TCHAR* BaseName, TCHAR* FullName );
+
+	/** Gets the exe file versions and lists all modules. */
+	void GetExeFileVersionAndModuleList( FCrashModuleInfo& out_ExeFileVersion );
+
+	/** Set the symbol paths based on the module paths. */
 	void SetSymbolPathsFromModules();
+
+	/** Gets detailed info about each module. */
 	void GetModuleInfoDetailed();
+
+	/** Check to see if the stack address resides within one of the loaded modules i.e. is it valid?. */
+	bool IsOffsetWithinModules( uint64 Offset );
+
+	/** Extract the system info of the crash from the minidump. */ // @TODO yrx 2015-02-24 Move to crash report client?
 	void GetSystemInfo();
-	void GetThreadInfo();
+
+	/** Extracts the thread info from the minidump. */
+	void GetThreadInfo(){}
+
+	/** Extracts info about the exception that caused the crash. */
 	void GetExceptionInfo();
-	bool GetCallstacks();
 
-	bool IsOffsetWithinModules( uint64 Offset );
+	/**
+	 * Gets the callstack of the crash.
+	 *
+	 * @return the number of valid function names
+	 */
+	int32 GetCallstacks();
 
-	static FString ExtractRelativePath( const TCHAR* BaseName, TCHAR* FullName );
+	/** Opens a minidump as a new session. */
+	bool OpenDumpFile( const FString& DumpFileName );
 
 protected:
 	/** Reference to the crash info. */
diff --git a/Engine/Source/Developer/CrashDebugHelper/Public/CrashDebugHelper.h b/Engine/Source/Developer/CrashDebugHelper/Public/CrashDebugHelper.h
--- a/Engine/Source/Developer/CrashDebugHelper/Public/CrashDebugHelper.h
+++ b/Engine/Source/Developer/CrashDebugHelper/Public/CrashDebugHelper.h
@@ -2,9 +2,9 @@
 
 #pragma once
 
+#include "Core.h"
+
 struct FPDBCacheEntry;
-typedef TSharedRef<FPDBCacheEntry> FPDBCacheEntryRef;
-typedef TSharedPtr<FPDBCacheEntry> FPDBCacheEntryPtr;
 
 enum EProcessorArchitecture
 {
@@ -120,6 +120,7 @@ public:
 	}
 };
 
+// @TODO yrx 2015-02-18 Merge with FCrashDescription
 /** A platform independent representation of a crash */
 class CRASHDEBUGHELPER_API FCrashInfo
 {
@@ -137,7 +138,7 @@ public:
 	FString DepotName;
 
 	/** Product version, based on FEngineVersion. */
-	FString ProductVersion;
+	FString EngineVersion;
 
 	/** CL built from. */
 	int32 BuiltFromCL;
@@ -152,7 +153,7 @@ public:
 	FString SymbolsPath;
 
 	FString SourceFile;
-	uint64 SourceLineNumber;
+	uint32 SourceLineNumber;
 	TArray<FString> SourceContext;
 
 	/** Only modules names, retrieved from the minidump file. */
@@ -164,7 +165,7 @@ public:
 	TArray<FCrashModuleInfo> Modules;
 
 	/** Shared pointer to the PDB Cache entry, if valid contains all information about synced PDBs. */
-	FPDBCacheEntryPtr PDBCacheEntry;
+	TSharedPtr<FPDBCacheEntry> PDBCacheEntry;
 
 	FCrashInfo()
 		: BuiltFromCL( INVALID_CHANGELIST )
@@ -187,13 +188,25 @@ public:
 	void Log( FString Line );
 
 private:
+	/** 
+	 * Convert the processor architecture to a human readable string
+	 */
 	const TCHAR* GetProcessorArchitecture( EProcessorArchitecture PA );
+
+	/**
+	 * Calculate the byte size of a UTF-8 string
+	 */
 	int64 StringSize( const ANSICHAR* Line );
+
+	/** 
+	* Write a line of UTF-8 to a file
+	*/
 	void WriteLine( FArchive* ReportFile, const ANSICHAR* Line = NULL );
 };
 
+
 /** Helper structure for tracking crash debug information */
-struct CRASHDEBUGHELPER_API FCrashDebugInfo
+struct FCrashDebugInfo
 {
 	/** The name of the crash dump file */
 	FString CrashDumpName;
@@ -205,236 +218,16 @@ struct CRASHDEBUGHELPER_API FCrashDebugInfo
 	FString SourceControlLabel;
 };
 
-/** Helper struct that holds various information about one PDB Cache entry. */
-struct FPDBCacheEntry
-{
-	/** Default constructor. */
-	FPDBCacheEntry( const FDateTime InLastAccessTime )
-		: LastAccessTime( InLastAccessTime )
-		, SizeGB( 0 )
-	{}
-
-	/** Initialization constructor. */
-	FPDBCacheEntry( const TArray<FString>& InFiles, const FString& InDirectory, const FDateTime InLastAccessTime, const int32 InSizeGB )
-		: Files( InFiles )
-		, Directory( InDirectory )
-		, LastAccessTime( InLastAccessTime )
-		, SizeGB( InSizeGB )
-	{}
-
-	void SetLastAccessTimeToNow()
-	{
-		LastAccessTime = FDateTime::UtcNow();
-	}
-
-	/** Paths to files associated with this PDB Cache entry. */
-	TArray<FString> Files;
-
-	/** The path associated with this PDB Cache entry. */
-	FString Directory;
-	
-	/** Last access time, changed every time this PDB cache entry is used. */
-	FDateTime LastAccessTime;
-
-	/** Size of the cache entry, in GBs. Rounded-up. */
-	const int32 SizeGB;
-
-	/**
-	 * Serializer.
-	 */
-	friend FArchive& operator<<(FArchive& Ar, FPDBCacheEntry& Entry)
-	{
-		return Ar << Entry.Files << (FString&)Entry.Directory << (int32&)Entry.SizeGB;
-	}
-};
-
-struct FPDBCacheEntryByAccessTime
-{
-	FORCEINLINE bool operator()( const FPDBCacheEntryRef& A, const FPDBCacheEntryRef& B ) const
-	{
-		return A->LastAccessTime.GetTicks() < B->LastAccessTime.GetTicks();
-	}
-};
-
-/** Implements PDB cache. */
-struct FPDBCache
-{
-protected:
-
-	// Defaults for the PDB cache.
-	enum
-	{
-		/** Size of the PDB cache, in GBs. */
-		PDB_CACHE_SIZE_GB = 128,
-		MIN_FREESPACE_GB = 64,
-		/** Age of file when it should be deleted from the PDB cache. */
-		DAYS_TO_DELETE_UNUSED_FILES = 14,
-
-		/**
-		*	Number of iterations inside the CleanPDBCache method.
-		*	Mostly to verify that MinDiskFreeSpaceGB requirement is met.
-		*/
-		CLEAN_PDBCACHE_NUM_ITERATIONS = 2,
-
-		/** Number of bytes per one gigabyte. */
-		NUM_BYTES_PER_GB = 1024 * 1024 * 1024
-	};
-
-	/** Dummy file used to read/set the file timestamp. */
-	static const TCHAR* PDBTimeStampFileNoMeta;
-
-	/** Data file used to read/set the file timestamp, contains all metadata. */
-	static const TCHAR* PDBTimeStampFile;
-
-	/** Map of the PDB Cache entries. */
-	TMap<FString, FPDBCacheEntryRef> PDBCacheEntries;
-
-	/** Path to the folder where the PDB cache is stored. */
-	FString PDBCachePath;
-
-	/** Age of file when it should be deleted from the PDB cache. */
-	int32 DaysToDeleteUnusedFilesFromPDBCache;
-
-	/** Size of the PDB cache, in GBs. */
-	int32 PDBCacheSizeGB;
-
-	/**
-	*	Minimum disk free space that should be available on the disk where the PDB cache is stored, in GBs.
-	*	Minidump diagnostics runs usually on the same drive as the crash reports drive, so we need to leave some space for the crash receiver.
-	*	If there is not enough disk free space, we will run the clean-up process.
-	*/
-	int32 MinDiskFreeSpaceGB;
-
-	/** Whether to use the PDB cache. */
-	bool bUsePDBCache;
-
-public:
-	/** Default constructor. */
-	FPDBCache()
-		: DaysToDeleteUnusedFilesFromPDBCache( DAYS_TO_DELETE_UNUSED_FILES )
-		, PDBCacheSizeGB( PDB_CACHE_SIZE_GB )
-		, MinDiskFreeSpaceGB( MIN_FREESPACE_GB )
-		, bUsePDBCache( false )
-	{}
-
-	/** Basic initialization, reading config etc.. */
-	void Init();
-
-	/**
-	 * @return whether to use the PDB cache.
-	 */
-	bool UsePDBCache() const
-	{
-		return bUsePDBCache;
-	}
-	
-	/**
-	 * @return true, if the PDB Cache contains the specified label.
-	 */
-	bool ContainsPDBCacheEntry( const FString& PathOrLabel ) const
-	{
-		return PDBCacheEntries.Contains( EscapePath( PathOrLabel ) );
-	}
-
-	/**
-	 *	Touches a PDB Cache entry by updating the timestamp.
-	 */
-	void TouchPDBCacheEntry( const FString& Directory );
-
-	/**
-	 * @return a PDB Cache entry for the specified label and touches it at the same time
-	 */
-	FPDBCacheEntryRef FindAndTouchPDBCacheEntry( const FString& PathOrLabel );
-
-	/**
-	 *	Creates a new PDB Cache entry, initializes it and adds to the database.
-	 */
-	FPDBCacheEntryRef CreateAndAddPDBCacheEntry( const FString& OriginalLabelName, const FString& DepotRoot, const FString& DepotName, const TArray<FString>& FilesToBeCached );
-
-	/**
-	 *	Creates a new PDB Cache entry, initializes it and adds to the database.
-	 */
-	FPDBCacheEntryRef CreateAndAddPDBCacheEntryMixed( const FString& ProductVersion, const TMap<FString,FString>& FilesToBeCached );
-
-protected:
-	/** Initializes the PDB Cache. */
-	void InitializePDBCache();
-
-	/**
-	 * @return the size of the PDB cache entry, in GBs.
-	 */
-	int32 GetPDBCacheEntrySizeGB( const FString& PathOrLabel ) const
-	{
-		return PDBCacheEntries.FindChecked( PathOrLabel )->SizeGB;
-	}
-
-	/**
-	 * @return the size of the PDB cache directory, in GBs.
-	 */
-	int32 GetPDBCacheSizeGB() const
-	{
-		int32 Result = 0;
-		if( bUsePDBCache )
-		{
-			for( const auto& It : PDBCacheEntries )
-			{
-				Result += It.Value->SizeGB;
-			}
-		}
-		return Result;
-	}
-
-	/**
-	 * Cleans the PDB Cache.
-	 *
-	 * @param DaysToKeep - removes all PDB Cache entries that are older than this value
-	 * @param NumberOfGBsToBeCleaned - if specifies, will try to remove as many PDB Cache entries as needed
-	 * to free disk space specified by this value
-	 *
-	 */
-	void CleanPDBCache( int32 DaysToKeep, int32 NumberOfGBsToBeCleaned = 0 );
-
-
-	/**
-	 *	Reads an existing PDB Cache entry.
-	 */
-	FPDBCacheEntryRef ReadPDBCacheEntry( const FString& Directory );
-
-	/**
-	 *	Sort PDB Cache entries by last access time.
-	 */
-	void SortPDBCache()
-	{
-		PDBCacheEntries.ValueSort( FPDBCacheEntryByAccessTime() );
-	}
-
-	/**
-	 *	Removes a PDB Cache entry from the database.
-	 *	Also removes all external files associated with this PDB Cache entry.
-	 */
-	void RemovePDBCacheEntry( const FString& Directory );
-
-	/** Replaces all invalid chars with + for the specified name. */
-	FString EscapePath( const FString& PathOrLabel ) const
-	{
-		// @see AutomationTool.CommandUtils.EscapePath 
-		return PathOrLabel.Replace( TEXT( ":" ), TEXT( "" ) ).Replace( TEXT( "/" ), TEXT( "+" ) ).Replace( TEXT( "\\" ), TEXT( "+" ) ).Replace( TEXT( " " ), TEXT( "+" ) );
-	}
-};
-
 /** The public interface for the crash dump handler singleton. */
 class CRASHDEBUGHELPER_API ICrashDebugHelper
 {
-protected:
-	/** The depot name that the handler is using to sync from */
-	FString DepotName;
+public:
+	static const TCHAR* P4_DEPOT_PREFIX;
 
+protected:
 	/** Depot root, only used in conjunction with the PDB Cache. */
 	FString DepotRoot;
 
-	/** The local folder where symbol files should be stored */
-	FString LocalSymbolStore;
-
 	/**
 	 *	Pattern to search in source control for the label.
 	 *	This somehow works for older crashes, before 4.2 and for the newest one,
@@ -444,18 +237,18 @@ protected:
 	FString SourceControlBuildLabelPattern;
 	
 	/**
-	 * Pattern to search in the network driver for the executables.
+	 * Patterns to search in the network driver for the executables/the symbols.
 	 * Valid from 4.2 UE builds. (CL-2068994)
 	 * This may change in future.
 	 */
-	FString ExecutablePathPattern;
+	TArray<FString> Branches;
+	TArray<FString> ExecutablePathPatterns;
+	TArray<FString> SymbolPathPatterns;
+
 	
 	/** Indicates that the crash handler is ready to do work */
 	bool bInitialized;
 
-	/** Instance of the PDB Cache. */
-	FPDBCache PDBCache;
-
 public:
 	/** A platform independent representation of a crash */
 	FCrashInfo CrashInfo;
@@ -471,18 +264,6 @@ public:
 	 */
 	virtual bool Init();
 
-	/** Get the depot name */
-	const FString& GetDepotName() const
-	{
-		return DepotName;
-	}
-
-	/** Set the depot name */
-	void SetDepotName( const FString& InDepotName )
-	{
-		DepotName = InDepotName;
-	}
-
 	/**
 	 *	Parse the given crash dump, determining EngineVersion of the build that produced it - if possible. 
 	 *
@@ -508,18 +289,6 @@ public:
 		return false;
 	}
 
-	/**
-	 *	Get the source control label for the given changelist number
-	 *
-	 *	@param	InChangelistNumber	The changelist number to retrieve the label for
-	 *
-	 *	@return	FString				The label containing that changelist number build, empty if not found
-	 */
-	virtual FString GetLabelFromChangelistNumber(int32 InChangelistNumber)
-	{
-		return RetrieveBuildLabel(InChangelistNumber);
-	}
-
 	/**
 	 *	Sync the required files from source control for debugging
 	 *
@@ -543,22 +312,7 @@ public:
 	virtual bool SyncRequiredFilesForDebuggingFromChangelist(int32 InChangelistNumber, const FString& InPlatform);
 
 	/**
-	 *	Sync required files and launch the debugger for the crash dump.
-	 *	Will call parse crash dump if it has not been called already.
-	 *
-	 *	@params	InCrashDumpName		The name of the crash dump file to debug
-	 *
-	 *	@return	bool				true if successful, false if not
-	 */
-	virtual bool SyncAndDebugCrashDump(const FString& InCrashDumpName)
-	{
-		return false;
-	}
-
-	/**
-	 *	Sync the binaries and associated pdbs to the requested label.
-	 *
-	 *	@return	bool		true if successful, false if not
+	 * Sync the branch root relative file names to the requested label
 	 */
 	virtual bool SyncModules();
 
@@ -578,21 +332,21 @@ public:
 	virtual bool AddAnnotatedSourceToReport();
 
 protected:
-	/**
-	 *	Retrieve the build label for the given changelist number.
-	 *
-	 *	@param	InChangelistNumber	The changelist number to retrieve the label for
-	 *	@return	FString				The label if successful, empty if it wasn't found or another error
-	 */
-	FString RetrieveBuildLabel(int32 InChangelistNumber);
+
+	/** Finds the storage of the symbols and the executables for the specified changelist and the depot name, it can be Perforce, network drive or stored locally. */
+	void FindSymbolsAndBinariesStorage();
 
 	/**
-	 *	Retrieve the build label and the network path for the given changelist number.
+	 *	Load the given ANSI text file to an array of strings - one FString per line of the file.
+	 *	Intended for use in simple text parsing actions
+	 *
+	 *	@param	InFilename			The text file to read, full path
+	 *	@param	OutStrings			The array of FStrings to fill in
+	 *
+	 *	@return	bool				true if successful, false if not
 	 */
-	void RetrieveBuildLabelAndNetworkPaths( int32 InChangelistNumber );
-
-
 	bool ReadSourceFile( const TCHAR* InFilename, TArray<FString>& OutStrings );
+
 public:
 	bool InitSourceControl(bool bShowLogin);
 	void ShutdownSourceControl();
diff --git a/Engine/Source/Developer/CrashDebugHelper/Public/CrashDebugHelperModule.h b/Engine/Source/Developer/CrashDebugHelper/Public/CrashDebugHelperModule.h
--- a/Engine/Source/Developer/CrashDebugHelper/Public/CrashDebugHelperModule.h
+++ b/Engine/Source/Developer/CrashDebugHelper/Public/CrashDebugHelperModule.h
@@ -2,16 +2,18 @@
 
 #pragma once
 
+#include "Core.h"
+#include "ModuleManager.h"
 #include "CrashDebugHelper.h"
 
-class CRASHDEBUGHELPER_API FCrashDebugHelperModule : public IModuleInterface
+class FCrashDebugHelperModule : public IModuleInterface
 {
 public:
 	virtual void StartupModule();
 	virtual void ShutdownModule();
 
 	/** Gets the crash debug helper singleton or returns NULL */
-	virtual ICrashDebugHelper* Get();
+	CRASHDEBUGHELPER_API virtual ICrashDebugHelper* Get();
 
 private:
 	class ICrashDebugHelper* CrashDebugHelper;
diff --git a/Engine/Source/Developer/CrashDebugHelper/Public/CrashDebugPDBCache.h b/Engine/Source/Developer/CrashDebugHelper/Public/CrashDebugPDBCache.h
new file mode 100644
--- /dev/null
+++ b/Engine/Source/Developer/CrashDebugHelper/Public/CrashDebugPDBCache.h
@@ -0,0 +1,232 @@
+// Copyright 1998-2015 Epic Games, Inc. All Rights Reserved.
+
+#pragma once
+
+#include "Core.h"
+
+typedef TSharedRef<FPDBCacheEntry> FPDBCacheEntryRef;
+typedef TSharedPtr<FPDBCacheEntry> FPDBCacheEntryPtr;
+
+/** Helper struct that holds various information about one PDB Cache entry. */
+struct FPDBCacheEntry
+{
+	/** Default constructor. */
+	FPDBCacheEntry( const FDateTime InLastAccessTime )
+		: LastAccessTime( InLastAccessTime )
+		, SizeGB( 0 )
+	{}
+
+	/** Initialization constructor. */
+	FPDBCacheEntry( const TArray<FString>& InFiles, const FString& InDirectory, const FDateTime InLastAccessTime, const int32 InSizeGB )
+		: Files( InFiles )
+		, Directory( InDirectory )
+		, LastAccessTime( InLastAccessTime )
+		, SizeGB( InSizeGB )
+	{}
+
+	void SetLastAccessTimeToNow()
+	{
+		LastAccessTime = FDateTime::UtcNow();
+	}
+
+	/** Paths to files associated with this PDB Cache entry. */
+	TArray<FString> Files;
+
+	/** The path associated with this PDB Cache entry. */
+	FString Directory;
+
+	/** Last access time, changed every time this PDB cache entry is used. */
+	FDateTime LastAccessTime;
+
+	/** Size of the cache entry, in GBs. Rounded-up. */
+	const int32 SizeGB;
+
+	/**
+	* Serializer.
+	*/
+	friend FArchive& operator<<(FArchive& Ar, FPDBCacheEntry& Entry)
+	{
+		return Ar << Entry.Files << (FString&)Entry.Directory << (int32&)Entry.SizeGB;
+	}
+};
+
+struct FPDBCacheEntryByAccessTime
+{
+	FORCEINLINE bool operator()( const FPDBCacheEntryRef& A, const FPDBCacheEntryRef& B ) const
+	{
+		return A->LastAccessTime.GetTicks() < B->LastAccessTime.GetTicks();
+	}
+};
+
+/** Implements PDB cache. */
+struct FPDBCache
+{
+protected:
+
+	// Defaults for the PDB cache.
+	enum
+	{
+		/** Size of the PDB cache, in GBs. */
+		PDB_CACHE_SIZE_GB = 128,
+		MIN_FREESPACE_GB = 64,
+		/** Age of file when it should be deleted from the PDB cache. */
+		DAYS_TO_DELETE_UNUSED_FILES = 14,
+
+		/**
+		*	Number of iterations inside the CleanPDBCache method.
+		*	Mostly to verify that MinDiskFreeSpaceGB requirement is met.
+		*/
+		CLEAN_PDBCACHE_NUM_ITERATIONS = 2,
+
+		/** Number of bytes per one gigabyte. */
+		NUM_BYTES_PER_GB = 1024 * 1024 * 1024
+	};
+
+	/** Dummy file used to read/set the file timestamp. */
+	static const TCHAR* PDBTimeStampFileNoMeta;
+
+	/** Data file used to read/set the file timestamp, contains all metadata. */
+	static const TCHAR* PDBTimeStampFile;
+
+	/** Map of the PDB Cache entries. */
+	TMap<FString, FPDBCacheEntryRef> PDBCacheEntries;
+
+	/** Path to the folder where the PDB cache is stored. */
+	FString PDBCachePath;
+
+	/** Age of file when it should be deleted from the PDB cache. */
+	int32 DaysToDeleteUnusedFilesFromPDBCache;
+
+	/** Size of the PDB cache, in GBs. */
+	int32 PDBCacheSizeGB;
+
+	/**
+	*	Minimum disk free space that should be available on the disk where the PDB cache is stored, in GBs.
+	*	Minidump diagnostics runs usually on the same drive as the crash reports drive, so we need to leave some space for the crash receiver.
+	*	If there is not enough disk free space, we will run the clean-up process.
+	*/
+	int32 MinDiskFreeSpaceGB;
+
+	/** Whether to use the PDB cache. */
+	bool bUsePDBCache;
+
+public:
+	/** Default constructor. */
+	FPDBCache()
+		: DaysToDeleteUnusedFilesFromPDBCache( DAYS_TO_DELETE_UNUSED_FILES )
+		, PDBCacheSizeGB( PDB_CACHE_SIZE_GB )
+		, MinDiskFreeSpaceGB( MIN_FREESPACE_GB )
+		, bUsePDBCache( false )
+	{}
+
+	/** Basic initialization, reading config etc.. */
+	void Init();
+
+	/**
+	* @return whether to use the PDB cache.
+	*/
+	bool UsePDBCache() const
+	{
+		return bUsePDBCache;
+	}
+
+	/** Accesses the singleton. */
+	CRASHDEBUGHELPER_API static FPDBCache& Get()
+	{
+		static FPDBCache Instance;
+		return Instance;
+	}
+
+	/**
+	* @return true, if the PDB Cache contains the specified label.
+	*/
+	bool ContainsPDBCacheEntry( const FString& PathOrLabel ) const
+	{
+		return PDBCacheEntries.Contains( EscapePath( PathOrLabel ) );
+	}
+
+	/**
+	*	Touches a PDB Cache entry by updating the timestamp.
+	*/
+	void TouchPDBCacheEntry( const FString& Directory );
+
+	/**
+	* @return a PDB Cache entry for the specified label and touches it at the same time
+	*/
+	FPDBCacheEntryRef FindAndTouchPDBCacheEntry( const FString& PathOrLabel );
+
+	/**
+	*	Creates a new PDB Cache entry, initializes it and adds to the database.
+	*/
+	FPDBCacheEntryRef CreateAndAddPDBCacheEntry( const FString& OriginalLabelName, const FString& DepotRoot, const FString& DepotName, const TArray<FString>& FilesToBeCached );
+
+	/**
+	*	Creates a new PDB Cache entry, initializes it and adds to the database.
+	*/
+	FPDBCacheEntryRef CreateAndAddPDBCacheEntryMixed( const FString& ProductVersion, const TMap<FString, FString>& FilesToBeCached );
+
+protected:
+	/** Initializes the PDB Cache. */
+	void InitializePDBCache();
+
+	/**
+	* @return the size of the PDB cache entry, in GBs.
+	*/
+	int32 GetPDBCacheEntrySizeGB( const FString& PathOrLabel ) const
+	{
+		return PDBCacheEntries.FindChecked( PathOrLabel )->SizeGB;
+	}
+
+	/**
+	* @return the size of the PDB cache directory, in GBs.
+	*/
+	int32 GetPDBCacheSizeGB() const
+	{
+		int32 Result = 0;
+		if( bUsePDBCache )
+		{
+			for( const auto& It : PDBCacheEntries )
+			{
+				Result += It.Value->SizeGB;
+			}
+		}
+		return Result;
+	}
+
+	/**
+	* Cleans the PDB Cache.
+	*
+	* @param DaysToKeep - removes all PDB Cache entries that are older than this value
+	* @param NumberOfGBsToBeCleaned - if specifies, will try to remove as many PDB Cache entries as needed
+	* to free disk space specified by this value
+	*
+	*/
+	void CleanPDBCache( int32 DaysToKeep, int32 NumberOfGBsToBeCleaned = 0 );
+
+
+	/**
+	*	Reads an existing PDB Cache entry.
+	*/
+	FPDBCacheEntryRef ReadPDBCacheEntry( const FString& Directory );
+
+	/**
+	*	Sort PDB Cache entries by last access time.
+	*/
+	void SortPDBCache()
+	{
+		PDBCacheEntries.ValueSort( FPDBCacheEntryByAccessTime() );
+	}
+
+	/**
+	*	Removes a PDB Cache entry from the database.
+	*	Also removes all external files associated with this PDB Cache entry.
+	*/
+	void RemovePDBCacheEntry( const FString& Directory );
+
+	/** Replaces all invalid chars with + for the specified name. */
+	FString EscapePath( const FString& PathOrLabel ) const
+	{
+		// @see AutomationTool.CommandUtils.EscapePath 
+		return PathOrLabel.Replace( TEXT( ":" ), TEXT( "" ) ).Replace( TEXT( "/" ), TEXT( "+" ) ).Replace( TEXT( "\\" ), TEXT( "+" ) ).Replace( TEXT( " " ), TEXT( "+" ) );
+	}
+};
\ No newline at end of file
diff --git a/Engine/Source/Programs/CrashReporter/MinidumpDiagnostics/Private/MinidumpDiagnosticsApp.cpp b/Engine/Source/Programs/CrashReporter/MinidumpDiagnostics/Private/MinidumpDiagnosticsApp.cpp
--- a/Engine/Source/Programs/CrashReporter/MinidumpDiagnostics/Private/MinidumpDiagnosticsApp.cpp
+++ b/Engine/Source/Programs/CrashReporter/MinidumpDiagnostics/Private/MinidumpDiagnosticsApp.cpp
@@ -23,7 +23,7 @@ int32 RunMinidumpDiagnostics(int32 ArgC, TCHAR* Argv[])
 	ICrashDebugHelper* CrashDebugHelper = CrashHelperModule.Get();
 	if( CrashDebugHelper == NULL )
 	{
-		UE_LOG( LogInit, Error, TEXT( " ... failed to load CrashDebugHelper module; unsupported platform?" ) );
+		UE_LOG( LogInit, Error, TEXT( "Failed to load CrashDebugHelper module; unsupported platform?" ) );
 		return 1;
 	}
 
@@ -33,11 +33,18 @@ int32 RunMinidumpDiagnostics(int32 ArgC, TCHAR* Argv[])
 
 	// Create a report for the minidump passed in on the command line
 	FString MinidumpName = Argv[1];
-	CrashDebugHelper->CreateMinidumpDiagnosticReport( MinidumpName );
+	const bool bValidCallstack = CrashDebugHelper->CreateMinidumpDiagnosticReport( MinidumpName );
+	FString DiagnosticsPath = FPaths::GetPath( MinidumpName );
+	if( bValidCallstack )
+	{	
+		DiagnosticsPath /= FString( TEXT( "Diagnostics.txt" ) );
+	}
+	else
+	{
+		DiagnosticsPath /= FString( TEXT( "DiagnosticsFailed.txt" ) );
+	}
 
-	// Write a report next to the original minidump
-	FString DiagnosticsPath = FPaths::GetPath(MinidumpName);
-	DiagnosticsPath /= FString( TEXT( "Diagnostics.txt" ) );
+	// Write a report next to the original minidump	
 	CrashDebugHelper->CrashInfo.GenerateReport( DiagnosticsPath );
 
 	// Cleanup
diff --git a/Engine/Source/Programs/SymbolDebugger/Private/SymbolDebugger.cpp b/Engine/Source/Programs/SymbolDebugger/Private/SymbolDebugger.cpp
--- a/Engine/Source/Programs/SymbolDebugger/Private/SymbolDebugger.cpp
+++ b/Engine/Source/Programs/SymbolDebugger/Private/SymbolDebugger.cpp
@@ -39,7 +39,7 @@ void FSymbolDebugger_AsyncInspect::DoWork()
 	}
 	else if ((EngineVersionName.Len() > 0) || (ChangelistName.Len() > 0))
 	{
-		const FString FoundLabel = CrashHelperModule->GetLabelFromChangelistNumber(FCString::Atoi(*ChangelistName));
+		const FString FoundLabel;// = CrashHelperModule->GetLabelFromChangelistNumber( FCString::Atoi( *ChangelistName ) );
 
 		if (FoundLabel.Len() > 0)
 		{
@@ -157,7 +157,7 @@ FSymbolDebugger::FSymbolDebugger()
 
 	CrashHelperModule = FModuleManager::LoadModuleChecked<FCrashDebugHelperModule>(FName("CrashDebugHelper")).Get();
 	check(CrashHelperModule != nullptr);
-	DepotName = CrashHelperModule->GetDepotName();
+	//DepotName = CrashHelperModule->GetDepotName();
 }
 
 bool FSymbolDebugger::SetCurrentMethod(SSymbolDebugger::ESymbolDebuggerMethods InNewMethod)
@@ -295,7 +295,7 @@ bool FSymbolDebugger::SetTextField(SSymbolDebugger::ESymbolDebuggerTextFields In
 	case SSymbolDebugger::TextField_SourceControlDepot:
 		{
 			DepotName = InNewName;
-			CrashHelperModule->SetDepotName(DepotName);
+			//CrashHelperModule->SetDepotName(DepotName);
 		}
 		break;
 	default:
diff --git a/Engine/Source/Runtime/Core/Private/GenericPlatform/GenericPlatformContext.cpp b/Engine/Source/Runtime/Core/Private/GenericPlatform/GenericPlatformContext.cpp
--- a/Engine/Source/Runtime/Core/Private/GenericPlatform/GenericPlatformContext.cpp
+++ b/Engine/Source/Runtime/Core/Private/GenericPlatform/GenericPlatformContext.cpp
@@ -234,3 +234,16 @@ void FGenericCrashContext::EndSection( const TCHAR* SectionName )
 	CommonBuffer += TEXT( ">" );
 	CommonBuffer += LINE_TERMINATOR;
 }
+
+
+FProgramCounterSymbolInfoEx::FProgramCounterSymbolInfoEx( FString InModuleName, FString InFunctionName, FString InFilename, uint32 InLineNumber, uint64 InSymbolDisplacement, uint64 InOffsetInModule, uint64 InProgramCounter ) :
+	ModuleName( InModuleName ),
+	FunctionName( InFunctionName ),
+	Filename( InFilename ),
+	LineNumber( InLineNumber ),
+	SymbolDisplacement( InSymbolDisplacement ),
+	OffsetInModule( InOffsetInModule ),
+	ProgramCounter( InProgramCounter )
+{
+
+}
\ No newline at end of file
diff --git a/Engine/Source/Runtime/Core/Private/GenericPlatform/GenericPlatformStackWalk.cpp b/Engine/Source/Runtime/Core/Private/GenericPlatform/GenericPlatformStackWalk.cpp
--- a/Engine/Source/Runtime/Core/Private/GenericPlatform/GenericPlatformStackWalk.cpp
+++ b/Engine/Source/Runtime/Core/Private/GenericPlatform/GenericPlatformStackWalk.cpp
@@ -1,8 +1,7 @@
 // Copyright 1998-2015 Epic Games, Inc. All Rights Reserved.
 
 #include "CorePrivatePCH.h"
-
-static const ANSICHAR* GUnknownFunction = "UnknownFunction";
+#include "GenericPlatformContext.h"
 
 FProgramCounterSymbolInfo::FProgramCounterSymbolInfo() :
 	LineNumber( 0 ),
@@ -10,9 +9,9 @@ FProgramCounterSymbolInfo::FProgramCounterSymbolInfo() :
 	OffsetInModule( 0 ),
 	ProgramCounter( 0 )
 {
-	FCStringAnsi::Strncpy( ModuleName, "UnknownModule", MAX_NAME_LENGHT );
-	FCStringAnsi::Strncpy( FunctionName, GUnknownFunction, MAX_NAME_LENGHT );
-	FCStringAnsi::Strncpy( Filename, "UnknownFile", MAX_NAME_LENGHT );
+	FCStringAnsi::Strncpy( ModuleName, "", MAX_NAME_LENGHT );
+	FCStringAnsi::Strncpy( FunctionName, "", MAX_NAME_LENGHT );
+	FCStringAnsi::Strncpy( Filename, "", MAX_NAME_LENGHT );
 }
 
 bool FGenericPlatformStackWalk::ProgramCounterToHumanReadableString( int32 CurrentCallDepth, uint64 ProgramCounter, ANSICHAR* HumanReadableString, SIZE_T HumanReadableStringSize, FGenericCrashContext* Context )
@@ -22,26 +21,101 @@ bool FGenericPlatformStackWalk::ProgramCounterToHumanReadableString( int32 Curre
 		FProgramCounterSymbolInfo SymbolInfo;
 		FPlatformStackWalk::ProgramCounterToSymbolInfo( ProgramCounter, SymbolInfo );
 
-		// ModuleName!FunctionName (ProgramCounter) + offset bytes [Filename:LineNumber]
-		ANSICHAR StackLine[MAX_SPRINTF];
+		return SymbolInfoToHumanReadableString( SymbolInfo, HumanReadableString, HumanReadableStringSize );
+	}
+	return false;
+}
+
+bool FGenericPlatformStackWalk::SymbolInfoToHumanReadableString( const FProgramCounterSymbolInfo& SymbolInfo, ANSICHAR* HumanReadableString, SIZE_T HumanReadableStringSize )
+{
+	const int32 MAX_TEMP_SPRINTF = 64;
+	// Valid callstack line 
+	// ModuleName!FunctionName {ProgramCounter} + offset bytes [Filename:LineNumber]
+	// 
+	// Invalid callstack line
+	// ModuleName! {ProgramCounter} + offset bytes
+	// 
+	if( HumanReadableString && HumanReadableStringSize > 0 )
+	{
+		ANSICHAR StackLine[MAX_SPRINTF] = {0};
 
 		// Strip module path.
 		const ANSICHAR* Pos0 = FCStringAnsi::Strrchr( SymbolInfo.ModuleName, '\\' );
 		const ANSICHAR* Pos1 = FCStringAnsi::Strrchr( SymbolInfo.ModuleName, '/' );
 		const UPTRINT RealPos = FMath::Max( (UPTRINT)Pos0, (UPTRINT)Pos1 );
-		const ANSICHAR* StrippedModuleName = RealPos > 0 ? (const ANSICHAR*)(RealPos+1) : SymbolInfo.ModuleName;
-	
-		FCStringAnsi::Sprintf( StackLine, "%s!%s (0x%016llx) + %i bytes [%s:%i]", StrippedModuleName, (const ANSICHAR*)SymbolInfo.FunctionName, SymbolInfo.ProgramCounter, SymbolInfo.SymbolDisplacement, (const ANSICHAR*)SymbolInfo.Filename, SymbolInfo.LineNumber );
+		const ANSICHAR* StrippedModuleName = RealPos > 0 ? (const ANSICHAR*)(RealPos + 1) : SymbolInfo.ModuleName;
+
+		//FCStringAnsi::Sprintf( StackLine, "%s!%s {0x%016llx} + %i bytes [%s:%i]", StrippedModuleName, (const ANSICHAR*)SymbolInfo.FunctionName, SymbolInfo.ProgramCounter, SymbolInfo.SymbolDisplacement, (const ANSICHAR*)SymbolInfo.Filename, SymbolInfo.LineNumber );
+		FCStringAnsi::Strcat( StackLine, MAX_SPRINTF, StrippedModuleName );
+		
+		const bool bHasValidFunctionName = FCStringAnsi::Strlen( SymbolInfo.FunctionName ) > 0;
+		if( bHasValidFunctionName )
+		{
+			FCStringAnsi::Strcat( StackLine, MAX_SPRINTF, "!" );
+			FCStringAnsi::Strcat( StackLine, MAX_SPRINTF, SymbolInfo.FunctionName );
+		}
+
+		{
+			ANSICHAR ProgramCounterAndSymbolDisplacement[MAX_TEMP_SPRINTF] = {0};
+			FCStringAnsi::Snprintf( ProgramCounterAndSymbolDisplacement, MAX_TEMP_SPRINTF, " {0x%016llx} + %i bytes", SymbolInfo.ProgramCounter, SymbolInfo.SymbolDisplacement );
+			FCStringAnsi::Strcat( StackLine, MAX_SPRINTF, ProgramCounterAndSymbolDisplacement );
+		}
+
+		const bool bHasValidFilename = FCStringAnsi::Strlen( SymbolInfo.Filename ) > 0 && SymbolInfo.LineNumber > 0;
+		if( bHasValidFilename )
+		{
+			ANSICHAR FilenameAndLineNumber[MAX_TEMP_SPRINTF] = {0};
+			FCStringAnsi::Snprintf( FilenameAndLineNumber, MAX_TEMP_SPRINTF, " [%s:%i]", SymbolInfo.Filename, SymbolInfo.LineNumber );
+			FCStringAnsi::Strcat( StackLine, MAX_SPRINTF, FilenameAndLineNumber );
+		}
+
 
 		// Append the stack line.
 		FCStringAnsi::Strcat( HumanReadableString, HumanReadableStringSize, StackLine );
 
-		// Return true, if we found a valid function name.
-		return FCStringAnsi::Strncmp( SymbolInfo.FunctionName, GUnknownFunction, FProgramCounterSymbolInfo::MAX_NAME_LENGHT ) != 0;
+		// Return true, if we have a valid function name.
+		return bHasValidFunctionName;
 	}
 	return false;
 }
 
+
+bool FGenericPlatformStackWalk::SymbolInfoToHumanReadableStringEx( const FProgramCounterSymbolInfoEx& SymbolInfo, FString& out_HumanReadableString )
+{
+	// Valid callstack line 
+	// ModuleName!FunctionName {ProgramCounter} + offset bytes [Filename:LineNumber]
+	// 
+	// Invalid callstack line
+	// ModuleName! {ProgramCounter} + offset bytes
+	
+	// Strip module path.
+	const TCHAR* Pos0 = FCString::Strrchr( *SymbolInfo.ModuleName, '\\' );
+	const TCHAR* Pos1 = FCString::Strrchr( *SymbolInfo.ModuleName, '/' );
+	const UPTRINT RealPos = FMath::Max( (UPTRINT)Pos0, (UPTRINT)Pos1 );
+	const FString StrippedModuleName = RealPos > 0 ? (const ANSICHAR*)(RealPos + 1) : SymbolInfo.ModuleName;
+
+	out_HumanReadableString = StrippedModuleName;
+	
+	const bool bHasValidFunctionName = SymbolInfo.FunctionName.Len() > 0;
+	if( bHasValidFunctionName )
+	{	
+		out_HumanReadableString += TEXT( "!" );
+		out_HumanReadableString += SymbolInfo.FunctionName;
+	}
+
+	//out_HumanReadableString += FString::Printf( TEXT( " {0x%016llx} + %i bytes" ), SymbolInfo.ProgramCounter, SymbolInfo.SymbolDisplacement );
+
+	const bool bHasValidFilename = SymbolInfo.Filename.Len() > 0 && SymbolInfo.LineNumber > 0;
+	if( bHasValidFilename )
+	{
+		out_HumanReadableString += FString::Printf( TEXT( " [%s:%i]" ), *SymbolInfo.Filename, SymbolInfo.LineNumber );
+	}
+
+	// Return true, if we have a valid function name.
+	return bHasValidFunctionName;
+}
+
+
 void FGenericPlatformStackWalk::CaptureStackBackTrace( uint64* BackTrace, uint32 MaxDepth, void* Context )
 {
 
diff --git a/Engine/Source/Runtime/Core/Public/GenericPlatform/GenericPlatformContext.h b/Engine/Source/Runtime/Core/Public/GenericPlatform/GenericPlatformContext.h
--- a/Engine/Source/Runtime/Core/Public/GenericPlatform/GenericPlatformContext.h
+++ b/Engine/Source/Runtime/Core/Public/GenericPlatform/GenericPlatformContext.h
@@ -9,7 +9,7 @@ class FString;
  * FString version.
  * To be used by external tools.
  */
-struct FProgramCounterSymbolInfoEx
+struct CORE_API FProgramCounterSymbolInfoEx
 {
 	/** Module name. */
 	FString	ModuleName;
diff --git a/Engine/Source/Runtime/Core/Public/GenericPlatform/GenericPlatformStackWalk.h b/Engine/Source/Runtime/Core/Public/GenericPlatform/GenericPlatformStackWalk.h
--- a/Engine/Source/Runtime/Core/Public/GenericPlatform/GenericPlatformStackWalk.h
+++ b/Engine/Source/Runtime/Core/Public/GenericPlatform/GenericPlatformStackWalk.h
@@ -8,6 +8,7 @@
 #pragma once
 #include "HAL/Platform.h"
 
+class FString;
 struct FGenericCrashContext;
 
 /**
@@ -34,9 +35,9 @@ struct FStackWalkModuleInfo
 };
 
 /**
- * Symbol information associated with a program counter.
+ * Symbol information associated with a program counter. ANSI version.
  */
-struct FProgramCounterSymbolInfo /*final*/
+struct FProgramCounterSymbolInfo final
 {
 	enum
 	{
@@ -69,6 +70,8 @@ struct FProgramCounterSymbolInfo /*final*/
 	FProgramCounterSymbolInfo();
 };
 
+struct FProgramCounterSymbolInfoEx;
+
 /**
  * Generic implementation for most platforms
  **/
@@ -105,6 +108,29 @@ struct CORE_API FGenericPlatformStackWalk
 	 */ 
 	static bool ProgramCounterToHumanReadableString( int32 CurrentCallDepth, uint64 ProgramCounter, ANSICHAR* HumanReadableString, SIZE_T HumanReadableStringSize, FGenericCrashContext* Context = nullptr );
 
+	/**
+	 * Converts the passed in symbol information to a human readable string and appends it to the passed in one.
+	 * @warning: The code assumes that HumanReadableString is large enough to contain the information.
+	 * @warning: Please, do not override this method. Can't be modified or altered without notice.
+	 * 
+	 * This method is the same for all platforms to simplify parsing by the crash processor.
+	 * 
+	 * Example formatted line:
+	 *
+	 * UE4Editor_Core!FOutputDeviceWindowsError::Serialize() (0xddf1bae5) + 620 bytes [\engine\source\runtime\core\private\windows\windowsplatformoutputdevices.cpp:110]
+	 * ModuleName!FunctionName (ProgramCounter) + offset bytes [StrippedFilepath:LineNumber]
+	 *
+	 * @param	SymbolInfo				Symbol information
+	 * @param	HumanReadableString		String to concatenate information with
+	 * @param	HumanReadableStringSize size of string in characters
+	 * @param	Context					Pointer to crash context, if any
+	 * @return	true if the symbol was found, otherwise false
+	 */ 
+	static bool SymbolInfoToHumanReadableString( const FProgramCounterSymbolInfo& SymbolInfo, ANSICHAR* HumanReadableString, SIZE_T HumanReadableStringSize );
+
+	/** Same as above, but can be used with external applications. */
+	static bool SymbolInfoToHumanReadableStringEx( const FProgramCounterSymbolInfoEx& SymbolInfo, FString& out_HumanReadableString );
+
 	/**
 	 * Converts the passed in program counter address to a symbol info struct, filling in module and filename, line number and displacement.
 	 * @warning: The code assumes that the destination strings are big enough
