CrashDebugHelper - Added support for indexed depot/PDB cache, will be used by multiple instance of the crash processor

diff --git a/Engine/Source/Developer/CrashDebugHelper/Private/CrashDebugHelper.cpp b/Engine/Source/Developer/CrashDebugHelper/Private/CrashDebugHelper.cpp
--- a/Engine/Source/Developer/CrashDebugHelper/Private/CrashDebugHelper.cpp
+++ b/Engine/Source/Developer/CrashDebugHelper/Private/CrashDebugHelper.cpp
@@ -16,6 +16,14 @@
 
 const TCHAR* ICrashDebugHelper::P4_DEPOT_PREFIX = TEXT( "//depot/" );
 
+void ICrashDebugHelper::SetDepotIndex( FString& PathToChange )
+{
+	FString CmdDepotIndex;
+	FParse::Value( FCommandLine::Get(), TEXT( "DepotIndex=" ), CmdDepotIndex );
+	// %DEPOT_INDEX% - Index of the depot, when multiple processor are used.
+	PathToChange.ReplaceInline( TEXT( "%DEPOT_INDEX%" ), *CmdDepotIndex );
+}
+
 bool ICrashDebugHelper::Init()
 {
 	bInitialized = true;
@@ -80,8 +88,11 @@ bool ICrashDebugHelper::Init()
 	UE_CLOG( !bCanUseSearchPatterns, LogCrashDebugHelper, Warning, TEXT( "Search patterns don't match" ) );
 
 	GConfig->GetString( TEXT( "Engine.CrashDebugHelper" ), TEXT( "DepotRoot" ), DepotRoot, GEngineIni );
+	ICrashDebugHelper::SetDepotIndex( DepotRoot );
+
 	const bool bHasDepotRoot = IFileManager::Get().DirectoryExists( *DepotRoot );
 	UE_CLOG( !bHasDepotRoot, LogCrashDebugHelper, Warning, TEXT( "DepotRoot: %s is not valid" ), *DepotRoot );
+	UE_LOG( LogCrashDebugHelper, Log, TEXT( "DepotRoot: %s" ), *DepotRoot );
 
 	if( bCanUseSearchPatterns && bHasDepotRoot )
 	{
diff --git a/Engine/Source/Developer/CrashDebugHelper/Private/CrashDebugPDBCache.cpp b/Engine/Source/Developer/CrashDebugHelper/Private/CrashDebugPDBCache.cpp
--- a/Engine/Source/Developer/CrashDebugHelper/Private/CrashDebugPDBCache.cpp
+++ b/Engine/Source/Developer/CrashDebugHelper/Private/CrashDebugPDBCache.cpp
@@ -33,6 +33,8 @@ void FPDBCache::Init()
 			UE_LOG( LogCrashDebugHelper, Warning, TEXT( "Failed to get PDBCachePath from ini file" ) );
 			bUsePDBCache = false;
 		}
+
+		ICrashDebugHelper::SetDepotIndex( PDBCachePath );
 	}
 
 	if( bUsePDBCache )
@@ -63,7 +65,7 @@ void FPDBCache::Init()
 		const int32 TotalDiscSpaceGB = int32( TotalNumberOfBytes >> 30 );
 		const int32 DiskFreeSpaceGB = int32( NumberOfFreeBytes >> 30 );
 
-		if( DiskFreeSpaceGB < MinDiskFreeSpaceGB )
+		if( DiskFreeSpaceGB < MinDiskFreeSpaceGB || TotalNumberOfBytes == 0 )
 		{
 			// There is not enough free space, calculate the current PDB cache usage and try removing the old data.
 			const int32 CurrentPDBCacheSizeGB = GetPDBCacheSizeGB();
@@ -80,17 +82,19 @@ void FPDBCache::Init()
 			else
 			{
 				// Clean the PDB cache until we get enough free space.
-				CleanPDBCache( DaysToDeleteUnusedFilesFromPDBCache, MinDiskFreeSpaceGB - DiskFreeSpaceGB );
+				const int32 MinSpaceRequirement = FMath::Max( MinDiskFreeSpaceGB - DiskFreeSpaceGB, 0 );
+				const int32 CacheSpaceRequirement = FMath::Max( CurrentPDBCacheSizeGB - PDBCacheSizeGB, 0 );
+				CleanPDBCache( DaysToDeleteUnusedFilesFromPDBCache, FMath::Max( MinSpaceRequirement, CacheSpaceRequirement ) );
 			}
 		}
 	}
 
 	if( bUsePDBCache )
 	{
-		UE_LOG( LogCrashDebugHelper, Log, TEXT( "PDBCachePath=%s" ), *PDBCachePath );
-		UE_LOG( LogCrashDebugHelper, Log, TEXT( "PDBCacheSizeGB=%i" ), PDBCacheSizeGB );
-		UE_LOG( LogCrashDebugHelper, Log, TEXT( "MinDiskFreeSpaceGB=%i" ), MinDiskFreeSpaceGB );
-		UE_LOG( LogCrashDebugHelper, Log, TEXT( "DaysToDeleteUnusedFilesFromPDBCache=%i" ), DaysToDeleteUnusedFilesFromPDBCache );
+		UE_LOG( LogCrashDebugHelper, Log, TEXT( "PDBCachePath: %s" ), *PDBCachePath );
+		UE_LOG( LogCrashDebugHelper, Log, TEXT( "PDBCacheSizeGB: %i" ), PDBCacheSizeGB );
+		UE_LOG( LogCrashDebugHelper, Log, TEXT( "MinDiskFreeSpaceGB: %i" ), MinDiskFreeSpaceGB );
+		UE_LOG( LogCrashDebugHelper, Log, TEXT( "DaysToDeleteUnusedFilesFromPDBCache: %i" ), DaysToDeleteUnusedFilesFromPDBCache );
 	}
 }
 
diff --git a/Engine/Source/Developer/CrashDebugHelper/Public/CrashDebugHelper.h b/Engine/Source/Developer/CrashDebugHelper/Public/CrashDebugHelper.h
--- a/Engine/Source/Developer/CrashDebugHelper/Public/CrashDebugHelper.h
+++ b/Engine/Source/Developer/CrashDebugHelper/Public/CrashDebugHelper.h
@@ -224,6 +224,9 @@ class CRASHDEBUGHELPER_API ICrashDebugHelper
 public:
 	static const TCHAR* P4_DEPOT_PREFIX;
 
+	/** Replaces %DEPOT_INDEX% with the command line DepotIndex in the specified path. */
+	static void SetDepotIndex( FString& PathToChange );
+
 protected:
 	/** Depot root, only used in conjunction with the PDB Cache. */
 	FString DepotRoot;
