CrashDebugHelper - Fixed source context showing bad source

diff --git a/Engine/Source/Developer/CrashDebugHelper/Private/CrashDebugHelper.cpp b/Engine/Source/Developer/CrashDebugHelper/Private/CrashDebugHelper.cpp
--- a/Engine/Source/Developer/CrashDebugHelper/Private/CrashDebugHelper.cpp
+++ b/Engine/Source/Developer/CrashDebugHelper/Private/CrashDebugHelper.cpp
@@ -404,19 +404,23 @@ bool ICrashDebugHelper::SyncSourceFile()
 }
 
 
-bool ICrashDebugHelper::ReadSourceFile( const TCHAR* InFilename, TArray<FString>& OutStrings )
+bool ICrashDebugHelper::ReadSourceFile( TArray<FString>& OutStrings )
 {
+	const FString FilePath = CrashInfo.DepotName.Replace( P4_DEPOT_PREFIX, *DepotRoot ) / CrashInfo.SourceFile;
+
 	FString Line;
-	if( FFileHelper::LoadFileToString( Line, InFilename ) )
+	if( FFileHelper::LoadFileToString( Line, *FilePath ) )
 	{
 		Line = Line.Replace( TEXT( "\r" ), TEXT( "" ) );
 		Line.ParseIntoArray( OutStrings, TEXT( "\n" ), false );
+		UE_LOG( LogCrashDebugHelper, Log, TEXT( "Reading a single source file: %s" ), *FilePath );
+		return false;
 		
 		return true;
 	}
 	else
 	{
-		UE_LOG( LogCrashDebugHelper, Warning, TEXT( "Failed to open source file %s" ), InFilename );
+		UE_LOG( LogCrashDebugHelper, Warning, TEXT( "Failed to open source file: %s" ), *FilePath );
 		return false;
 	}
 }
@@ -426,8 +430,7 @@ void ICrashDebugHelper::AddSourceToReport()
 	if( CrashInfo.SourceFile.Len() > 0 && CrashInfo.SourceLineNumber != 0 )
 	{
 		TArray<FString> Lines;
-		FString FullPath = FString( TEXT( "../../../" ) ) + CrashInfo.SourceFile;
-		ReadSourceFile( *FullPath, Lines );
+		ReadSourceFile( Lines );
 
 		const uint32 MinLine = FMath::Clamp( CrashInfo.SourceLineNumber - 15, (uint32)1, (uint32)Lines.Num() );
 		const uint32 MaxLine = FMath::Clamp( CrashInfo.SourceLineNumber + 15, (uint32)1, (uint32)Lines.Num() );
@@ -436,11 +439,11 @@ void ICrashDebugHelper::AddSourceToReport()
 		{
 			if( Line == CrashInfo.SourceLineNumber - 1 )
 			{
-				CrashInfo.SourceContext.Add( FString( TEXT( "*****" ) ) + Lines[Line] );
+				CrashInfo.SourceContext.Add( FString::Printf( TEXT( "%5u ***** %s" ), Line, *Lines[Line] ) );
 			}
 			else
 			{
-				CrashInfo.SourceContext.Add( FString( TEXT( "     " ) ) + Lines[Line] );
+				CrashInfo.SourceContext.Add( FString::Printf( TEXT( "%5u       %s" ), Line, *Lines[Line] ) );
 			}
 		}
 	}
diff --git a/Engine/Source/Developer/CrashDebugHelper/Public/CrashDebugHelper.h b/Engine/Source/Developer/CrashDebugHelper/Public/CrashDebugHelper.h
--- a/Engine/Source/Developer/CrashDebugHelper/Public/CrashDebugHelper.h
+++ b/Engine/Source/Developer/CrashDebugHelper/Public/CrashDebugHelper.h
@@ -343,12 +343,11 @@ protected:
 	 *	Load the given ANSI text file to an array of strings - one FString per line of the file.
 	 *	Intended for use in simple text parsing actions
 	 *
-	 *	@param	InFilename			The text file to read, full path
 	 *	@param	OutStrings			The array of FStrings to fill in
 	 *
 	 *	@return	bool				true if successful, false if not
 	 */
-	bool ReadSourceFile( const TCHAR* InFilename, TArray<FString>& OutStrings );
+	bool ReadSourceFile( TArray<FString>& OutStrings );
 
 public:
 	bool InitSourceControl(bool bShowLogin);
