Fix CrashReporter shadow variables

diff --git a/Engine/Source/Developer/CrashDebugHelper/Private/Mac/CrashDebugHelperMac.cpp b/Engine/Source/Developer/CrashDebugHelper/Private/Mac/CrashDebugHelperMac.cpp
--- a/Engine/Source/Developer/CrashDebugHelper/Private/Mac/CrashDebugHelperMac.cpp
+++ b/Engine/Source/Developer/CrashDebugHelper/Private/Mac/CrashDebugHelperMac.cpp
@@ -526,7 +526,7 @@ bool FCrashDebugHelperMac::ParseCrashDump(const FString& InCrashDumpName, FCrash
 			int32 Build = 0;
 			int32 CLNumber = 0;
 			FString Branch;
-			int32 Result = ParseVersion(*CrashDump, Major, Minor, Build, CLNumber, Branch);
+			Result = ParseVersion(*CrashDump, Major, Minor, Build, CLNumber, Branch);
 			if(Result >= 3)
 			{
 				if (Result < 5)
@@ -557,18 +557,20 @@ bool FCrashDebugHelperMac::CreateMinidumpDiagnosticReport( const FString& InCras
 	
 	FString CrashDump;
 	
-	NSString* CrashDumpPath = InCrashDumpName.GetNSString();
-	NSError* Error = nil;
-    NSData* Data = [NSData dataWithContentsOfFile: CrashDumpPath options: NSMappedRead error: &Error];
-	if(Data && !Error)
-	{
-		PLCrashReport* CrashLog = [[PLCrashReport alloc] initWithData: Data error: &Error];
-		if(CrashLog && !Error)
-		{
-			NSString* Report = [PLCrashReportTextFormatter stringValueForCrashReport: CrashLog withTextFormat: PLCrashReportTextFormatiOS];
-			CrashDump = FString(Report);
-		}
-	}
+    {
+        NSString* CrashDumpPath = InCrashDumpName.GetNSString();
+        NSError* Error = nil;
+        NSData* Data = [NSData dataWithContentsOfFile: CrashDumpPath options: NSMappedRead error: &Error];
+        if(Data && !Error)
+        {
+            PLCrashReport* CrashLog = [[PLCrashReport alloc] initWithData: Data error: &Error];
+            if(CrashLog && !Error)
+            {
+                NSString* Report = [PLCrashReportTextFormatter stringValueForCrashReport: CrashLog withTextFormat: PLCrashReportTextFormatiOS];
+                CrashDump = FString(Report);
+            }
+        }
+    }
 	
 	if ( !CrashDump.IsEmpty() || FFileHelper::LoadFileToString( CrashDump, *InCrashDumpName ) )
 	{
@@ -592,7 +594,7 @@ bool FCrashDebugHelperMac::CreateMinidumpDiagnosticReport( const FString& InCras
 			int32 CLNumber = 0;
 			int32 LineNumber = 0;;
 			
-			int32 Result = ParseVersion(*CrashDump, Major, Minor, Build, CLNumber, Branch);
+			Result = ParseVersion(*CrashDump, Major, Minor, Build, CLNumber, Branch);
 			if(Result >= 3)
 			{
 				CrashInfo.EngineVersion = FEngineVersion(Major, Minor, Build, CLNumber, Branch).ToString();
@@ -638,8 +640,7 @@ bool FCrashDebugHelperMac::CreateMinidumpDiagnosticReport( const FString& InCras
 			while(ModuleLine)
 			{
 				FCrashModuleInfo Module;
-				bool const bOK = ParseModuleLine(ModuleLine, Module);
-				if(bOK)
+				if (ParseModuleLine(ModuleLine, Module))
 				{
 					CrashInfo.Modules.Push(Module);
 					CrashInfo.ModuleNames.Push(FPaths::GetBaseFilename(Module.Name));
diff --git a/Engine/Source/Developer/StandaloneRenderer/Private/Mac/OpenGL/SlateOpenGLViewport.cpp b/Engine/Source/Developer/StandaloneRenderer/Private/Mac/OpenGL/SlateOpenGLViewport.cpp
--- a/Engine/Source/Developer/StandaloneRenderer/Private/Mac/OpenGL/SlateOpenGLViewport.cpp
+++ b/Engine/Source/Developer/StandaloneRenderer/Private/Mac/OpenGL/SlateOpenGLViewport.cpp
@@ -55,7 +55,7 @@ void FSlateOpenGLViewport::SwapBuffers()
 	ContextGuard.Unlock();
 }
 
-void FSlateOpenGLViewport::Resize(int32 Width, int32 Height, bool bFullscreen)
+void FSlateOpenGLViewport::Resize(int32 Width, int32 Height, bool bInFullscreen)
 {
 	FScopeLock Lock(&ContextGuard);
 
diff --git a/Engine/Source/Developer/StandaloneRenderer/Private/OpenGL/SlateOpenGLRenderer.cpp b/Engine/Source/Developer/StandaloneRenderer/Private/OpenGL/SlateOpenGLRenderer.cpp
--- a/Engine/Source/Developer/StandaloneRenderer/Private/OpenGL/SlateOpenGLRenderer.cpp
+++ b/Engine/Source/Developer/StandaloneRenderer/Private/OpenGL/SlateOpenGLRenderer.cpp
@@ -218,11 +218,11 @@ bool FSlateOpenGLRenderer::GenerateDynamicImageResource(FName ResourceName, uint
 }
 
 
-void FSlateOpenGLRenderer::LoadStyleResources( const ISlateStyle& Style )
+void FSlateOpenGLRenderer::LoadStyleResources( const ISlateStyle& InStyle )
 {
 	if ( TextureManager.IsValid() )
 	{
-		TextureManager->LoadStyleResources( Style );
+		TextureManager->LoadStyleResources( InStyle );
 	}
 }
 
diff --git a/Engine/Source/Developer/StandaloneRenderer/Private/OpenGL/SlateOpenGLTextureManager.cpp b/Engine/Source/Developer/StandaloneRenderer/Private/OpenGL/SlateOpenGLTextureManager.cpp
--- a/Engine/Source/Developer/StandaloneRenderer/Private/OpenGL/SlateOpenGLTextureManager.cpp
+++ b/Engine/Source/Developer/StandaloneRenderer/Private/OpenGL/SlateOpenGLTextureManager.cpp
@@ -10,9 +10,9 @@
 
 DEFINE_LOG_CATEGORY_STATIC(LogSlateOpenGL, Log, All);
 
-FSlateOpenGLTextureManager::FDynamicTextureResource::FDynamicTextureResource( FSlateOpenGLTexture* OpenGLTexture )
+FSlateOpenGLTextureManager::FDynamicTextureResource::FDynamicTextureResource( FSlateOpenGLTexture* InOpenGLTexture )
 	: Proxy( new FSlateShaderResourceProxy )
-	, OpenGLTexture( OpenGLTexture  )
+	, OpenGLTexture( InOpenGLTexture  )
 {
 }
 
diff --git a/Engine/Source/Programs/CrashReporter/CrashReportClient/Private/CrashDescription.cpp b/Engine/Source/Programs/CrashReporter/CrashReportClient/Private/CrashDescription.cpp
--- a/Engine/Source/Programs/CrashReporter/CrashReportClient/Private/CrashDescription.cpp
+++ b/Engine/Source/Programs/CrashReporter/CrashReportClient/Private/CrashDescription.cpp
@@ -129,10 +129,10 @@ FCrashDescription::FCrashDescription( FString WERXMLFilepath ) :
 
 				if( ParsedParameters9.Num() > 1 )
 				{
-					const FString BaseDir = ParsedParameters9[1];
+					const FString BaseDirectory = ParsedParameters9[1];
 
 					TArray<FString> SubDirs;
-					BaseDir.ParseIntoArray( SubDirs, TEXT( "/" ), true );
+					BaseDirectory.ParseIntoArray( SubDirs, TEXT( "/" ), true );
 					const int SubDirsNum = SubDirs.Num();
 					const FString PlatformName = SubDirsNum > 0 ? SubDirs[SubDirsNum - 1] : TEXT("");
 					if( Product.Len() > 0 )
diff --git a/Engine/Source/Programs/CrashReporter/CrashReportClient/Private/MainLoopTiming.cpp b/Engine/Source/Programs/CrashReporter/CrashReportClient/Private/MainLoopTiming.cpp
--- a/Engine/Source/Programs/CrashReporter/CrashReportClient/Private/MainLoopTiming.cpp
+++ b/Engine/Source/Programs/CrashReporter/CrashReportClient/Private/MainLoopTiming.cpp
@@ -6,8 +6,8 @@
 #include "LaunchEngineLoop.h"
 #include "TaskGraphInterfaces.h"
 
-FMainLoopTiming::FMainLoopTiming(float IdealTickRate, EMainLoopOptions::Type Options)
-	: IdealFrameTime(1.f / IdealTickRate)
+FMainLoopTiming::FMainLoopTiming(float InIdealTickRate, EMainLoopOptions::Type Options)
+	: IdealFrameTime(1.f / InIdealTickRate)
 	, ActualDeltaTime(0)
 	, LastTime(FPlatformTime::Seconds())
 	, bTickSlate(Options & EMainLoopOptions::UsingSlate)
