MDD - Fixed PDB cache not working

diff --git a/Engine/Source/Developer/CrashDebugHelper/Private/CrashDebugPDBCache.cpp b/Engine/Source/Developer/CrashDebugHelper/Private/CrashDebugPDBCache.cpp
--- a/Engine/Source/Developer/CrashDebugHelper/Private/CrashDebugPDBCache.cpp
+++ b/Engine/Source/Developer/CrashDebugHelper/Private/CrashDebugPDBCache.cpp
@@ -105,7 +105,7 @@ void FPDBCache::InitializePDBCache()
 	IFileManager::Get().MakeDirectory( *PDBCachePath, true );
 
 	TArray<FString> PDBCacheEntryDirectories;
-	IFileManager::Get().FindFiles( PDBCacheEntryDirectories, *PDBCachePath, false, true );
+	IFileManager::Get().FindFiles( PDBCacheEntryDirectories, *(PDBCachePath / TEXT( "*" )), false, true );
 
 	for( const auto& Directory : PDBCacheEntryDirectories )
 	{
