CrashDebugHelper - Annotate for source files

diff --git a/Engine/Source/Developer/CrashDebugHelper/Private/CrashDebugHelper.cpp b/Engine/Source/Developer/CrashDebugHelper/Private/CrashDebugHelper.cpp
--- a/Engine/Source/Developer/CrashDebugHelper/Private/CrashDebugHelper.cpp
+++ b/Engine/Source/Developer/CrashDebugHelper/Private/CrashDebugHelper.cpp
@@ -464,7 +464,7 @@ bool ICrashDebugHelper::AddAnnotatedSourceToReport()
 		FString DepotPath = CrashInfo.DepotName / CrashInfo.SourceFile;
 
 		TArray<FAnnotationLine> Lines;
-		SourceControlHelpers::AnnotateFile( ISourceControlModule::Get().GetProvider(), CrashInfo.LabelName, DepotPath, Lines );
+		SourceControlHelpers::AnnotateFile( ISourceControlModule::Get().GetProvider(), CrashInfo.BuiltFromCL, DepotPath, Lines );
 
 		uint32 MinLine = FMath::Clamp( CrashInfo.SourceLineNumber - 15, (uint32)1, (uint32)Lines.Num() );
 		uint32 MaxLine = FMath::Clamp( CrashInfo.SourceLineNumber + 15, (uint32)1, (uint32)Lines.Num() );
