CrashDebugHelper - Added line numbers for annotated files

diff --git a/Engine/Source/Developer/CrashDebugHelper/Private/CrashDebugHelper.cpp b/Engine/Source/Developer/CrashDebugHelper/Private/CrashDebugHelper.cpp
--- a/Engine/Source/Developer/CrashDebugHelper/Private/CrashDebugHelper.cpp
+++ b/Engine/Source/Developer/CrashDebugHelper/Private/CrashDebugHelper.cpp
@@ -474,11 +474,11 @@ bool ICrashDebugHelper::AddAnnotatedSourceToReport()
 		{			
 			if( Line == CrashInfo.SourceLineNumber )
 			{
-				CrashInfo.SourceContext.Add( FString::Printf( TEXT( "*****%20s: %s" ), *Lines[Line].UserName, *Lines[Line].Line ) );
+				CrashInfo.SourceContext.Add( FString::Printf( TEXT( "%5u ***** %20s: %s" ), Line, *Lines[Line].UserName, *Lines[Line].Line ) );
 			}
 			else
 			{
-				CrashInfo.SourceContext.Add( FString::Printf( TEXT( "     %20s: %s" ), *Lines[Line].UserName, *Lines[Line].Line ) );
+				CrashInfo.SourceContext.Add( FString::Printf( TEXT( "%5u       %20s: %s" ), Line, *Lines[Line].UserName, *Lines[Line].Line ) );
 			}
 		}
 		return true;
