Core - Removes old UE4 crash contexts, better processing for 'looped' ensures

diff --git a/Engine/Source/Developer/CrashDebugHelper/Private/Windows/WindowsPlatformStackWalkExt.cpp b/Engine/Source/Developer/CrashDebugHelper/Private/Windows/WindowsPlatformStackWalkExt.cpp
--- a/Engine/Source/Developer/CrashDebugHelper/Private/Windows/WindowsPlatformStackWalkExt.cpp
+++ b/Engine/Source/Developer/CrashDebugHelper/Private/Windows/WindowsPlatformStackWalkExt.cpp
@@ -570,7 +570,7 @@ int FWindowsPlatformStackWalkExt::GetCallstacks()
 	}
 
 	// Get the entire stack trace
-	const uint32 MaxFrames = 128;
+	const uint32 MaxFrames = 8192;
 	const uint32 MaxFramesSize = MaxFrames * ContextUsed;
 
 	DEBUG_STACK_FRAME* StackFrames = new(FMemStack::Get()) DEBUG_STACK_FRAME[MaxFrames];
@@ -580,6 +580,8 @@ int FWindowsPlatformStackWalkExt::GetCallstacks()
 	FMemory::Memzero( ContextData, MaxFramesSize );
 	HRESULT HR = Control->GetContextStackTrace( Context, ContextUsed, StackFrames, MaxFrames, ContextData, MaxFramesSize, ContextUsed, &Count );
 
+	int32 AssertOnEnsureIndex = -1;
+
 	for( uint32 StackIndex = 0; StackIndex < Count; StackIndex++ )
 	{	
 		const uint64 Offset = StackFrames[StackIndex].InstructionOffset;
@@ -631,11 +633,11 @@ int FWindowsPlatformStackWalkExt::GetCallstacks()
 			// If we find an assert, the actual source file we're interested in is the next one up, so reset the source file found flag
 			if( FunctionName.Len() > 0 )
 			{
-				if( FunctionName.Contains( TEXT( "FDebug::AssertFailed" ), ESearchCase::CaseSensitive )
-					|| FunctionName.Contains( TEXT( "FDebug::EnsureNotFalse" ), ESearchCase::CaseSensitive )
-					|| FunctionName.Contains( TEXT( "FDebug::EnsureFailed" ), ESearchCase::CaseSensitive ) )
+				if( FunctionName.Contains( TEXT( "FDebug::" ), ESearchCase::CaseSensitive )
+					|| FunctionName.Contains( TEXT( "NewReportEnsure" ), ESearchCase::CaseSensitive ) )
 				{
 					bFoundSourceFile = false;
+					AssertOnEnsureIndex = FMath::Max( AssertOnEnsureIndex, (int32)StackIndex );
 				}
 			}
 
@@ -649,6 +651,13 @@ int FWindowsPlatformStackWalkExt::GetCallstacks()
 		}
 	}
 
+	// Remove callstack entries below FDebug, we don't need them.
+	if (AssertOnEnsureIndex > 0)
+	{	
+		Exception.CallStackString.RemoveAt( 0, AssertOnEnsureIndex );
+		UE_LOG( LogCrashDebugHelper, Warning, TEXT( "Callstack trimmed to %i entries" ), Exception.CallStackString.Num() );
+	}
+
 	UE_LOG( LogCrashDebugHelper, Warning, TEXT( "Callstack generated with %i valid function names" ), NumValidFunctionNames );
 
 	return NumValidFunctionNames;
diff --git a/Engine/Source/Runtime/Core/Private/Misc/CoreMisc.cpp b/Engine/Source/Runtime/Core/Private/Misc/CoreMisc.cpp
--- a/Engine/Source/Runtime/Core/Private/Misc/CoreMisc.cpp
+++ b/Engine/Source/Runtime/Core/Private/Misc/CoreMisc.cpp
@@ -697,6 +697,21 @@ void FMaintenance::DeleteOldLogs()
 				IFileManager::Get().Delete(*FullFileName);
 			}
 		}
+
+		// Remove old UE4 crash contexts
+		TArray<FString> Directories;
+		IFileManager::Get().FindFiles( Directories, *FString::Printf( TEXT( "%s/UE4CC*" ), *FPaths::GameLogDir() ), false, true );
+
+		for (auto Dir : Directories)
+		{
+			const FString CrashContextDirectory = FPaths::GameLogDir() / Dir;
+			const FDateTime DirectoryAccessTime = IFileManager::Get().GetTimeStamp( *CrashContextDirectory );
+			if (FDateTime::Now() - DirectoryAccessTime > FTimespan::FromDays( PurgeLogsDays ))
+			{
+				UE_LOG( LogStreaming, Log, TEXT( "Deleting old crash context %s" ), *Dir );
+				IFileManager::Get().DeleteDirectory( *CrashContextDirectory, false, true );
+			}
+		}
 	}
 }
 
