CrashDebugHelper - Fixed source context not showing the source

diff --git a/Engine/Source/Developer/CrashDebugHelper/Private/CrashDebugHelper.cpp b/Engine/Source/Developer/CrashDebugHelper/Private/CrashDebugHelper.cpp
--- a/Engine/Source/Developer/CrashDebugHelper/Private/CrashDebugHelper.cpp
+++ b/Engine/Source/Developer/CrashDebugHelper/Private/CrashDebugHelper.cpp
@@ -55,10 +55,10 @@ bool ICrashDebugHelper::Init()
 		// Default to BRANCH_NAME
 		else
 		{
-			DepotName = FString::Printf( TEXT( "%s%s" ), P4_DEPOT_PREFIX, TEXT( BRANCH_NAME ) );
+			DepotName = FString::Printf( TEXT( "%s" ), TEXT( BRANCH_NAME ) );
 		}
 
-		CrashInfo.DepotName = DepotName;
+		CrashInfo.DepotName = DepotName.Replace( TEXT( "+" ), TEXT( "/" ) );
 
 		// Try to get the BuiltFromCL from command line to use this instead of attempting to locate the CL in the minidump
 		FString CmdLineBuiltFromCL;
@@ -78,6 +78,11 @@ bool ICrashDebugHelper::Init()
 
 		CrashInfo.BuiltFromCL = BuiltFromCL;
 	}
+
+
+	UE_LOG( LogCrashDebugHelper, Log, TEXT( "DepotName: %s" ), *CrashInfo.DepotName );
+	UE_LOG( LogCrashDebugHelper, Log, TEXT( "BuiltFromCL: %i" ), CrashInfo.BuiltFromCL );
+	UE_LOG( LogCrashDebugHelper, Log, TEXT( "EngineVersion: %s" ), *CrashInfo.EngineVersion );
 	
 	GConfig->GetString( TEXT( "Engine.CrashDebugHelper" ), TEXT( "SourceControlBuildLabelPattern" ), SourceControlBuildLabelPattern, GEngineIni );
 
@@ -413,9 +418,7 @@ bool ICrashDebugHelper::ReadSourceFile( TArray<FString>& OutStrings )
 	{
 		Line = Line.Replace( TEXT( "\r" ), TEXT( "" ) );
 		Line.ParseIntoArray( OutStrings, TEXT( "\n" ), false );
-		UE_LOG( LogCrashDebugHelper, Log, TEXT( "Reading a single source file: %s" ), *FilePath );
-		return false;
-		
+		UE_LOG( LogCrashDebugHelper, Log, TEXT( "Reading a single source file: %s" ), *FilePath );		
 		return true;
 	}
 	else
