Crash report
UECORE-164/Source context is not working properly for crashes from the UE4 releases
UE-13693/frequent crashes give "No disk in drive" error
UE-15369/UE-Various: Been getting "There is no Disk in the drive" crashreporter error. On orion the message is there permeneantly.

diff --git a/Engine/Source/Developer/CrashDebugHelper/Private/CrashDebugHelper.cpp b/Engine/Source/Developer/CrashDebugHelper/Private/CrashDebugHelper.cpp
--- a/Engine/Source/Developer/CrashDebugHelper/Private/CrashDebugHelper.cpp
+++ b/Engine/Source/Developer/CrashDebugHelper/Private/CrashDebugHelper.cpp
@@ -92,14 +92,7 @@ bool ICrashDebugHelper::Init()
 	const bool bCanUseSearchPatterns = Branches.Num() == ExecutablePathPatterns.Num() && ExecutablePathPatterns.Num() == SymbolPathPatterns.Num() && Branches.Num() > 0;
 	UE_CLOG( !bCanUseSearchPatterns, LogCrashDebugHelper, Warning, TEXT( "Search patterns don't match" ) );
 
-	GConfig->GetString( TEXT( "Engine.CrashDebugHelper" ), TEXT( "DepotRoot" ), DepotRoot, GEngineIni );
-	ICrashDebugHelper::SetDepotIndex( DepotRoot );
-
-	const bool bHasDepotRoot = IFileManager::Get().DirectoryExists( *DepotRoot );
-	UE_CLOG( !bHasDepotRoot, LogCrashDebugHelper, Warning, TEXT( "DepotRoot: %s is not valid" ), *DepotRoot );
-	UE_LOG( LogCrashDebugHelper, Log, TEXT( "DepotRoot: %s" ), *DepotRoot );
-
-	if( bCanUseSearchPatterns && bHasDepotRoot )
+	if (bCanUseSearchPatterns)
 	{
 		FPDBCache::Get().Init();
 	}
@@ -279,7 +272,7 @@ bool ICrashDebugHelper::SyncModules()
 				for( const auto& PDBPath : PDBPaths )
 				{
 					const FString PDBRelativePath = PDBPath.Replace( *CrashInfo.DepotName, TEXT( "" ) ).Replace( UESymbols, TEXT( "" ) );
-					const FString PDBFullpath = DepotRoot / PDBPath.Replace( P4_DEPOT_PREFIX, TEXT( "" ) );
+					const FString PDBFullpath = FPDBCache::Get().GetDepotRoot() / PDBPath.Replace( P4_DEPOT_PREFIX, TEXT( "" ) );
 
 					const FString PDBMatch = PDBRelativePath.Replace( TEXT( "pdb" ), TEXT( "" ) );
 					const FString NetworkRelativePath = NetworkExecutableFullpath.Replace( *CrashInfo.ExecutablesPath, TEXT( "" ) );
@@ -379,7 +372,7 @@ bool ICrashDebugHelper::SyncModules()
 			}
 
 			// Initialize and add a new PDB Cache entry to the database.
-			CrashInfo.PDBCacheEntry = FPDBCache::Get().CreateAndAddPDBCacheEntry( CrashInfo.LabelName, DepotRoot, CrashInfo.DepotName, FilesToBeCached );
+			CrashInfo.PDBCacheEntry = FPDBCache::Get().CreateAndAddPDBCacheEntry( CrashInfo.LabelName, CrashInfo.DepotName, FilesToBeCached );
 		}
 	}
 	else
@@ -411,14 +404,24 @@ bool ICrashDebugHelper::SyncSourceFile()
 
 bool ICrashDebugHelper::ReadSourceFile( TArray<FString>& OutStrings )
 {
-	const FString FilePath = CrashInfo.DepotName.Replace( P4_DEPOT_PREFIX, *DepotRoot ) / CrashInfo.SourceFile;
+	const bool bUsePDBCache = FPDBCache::Get().UsePDBCache();
+
+	FString FilePath;
+	if (bUsePDBCache)
+	{
+		FilePath = CrashInfo.DepotName.Replace( P4_DEPOT_PREFIX, *FPDBCache::Get().GetDepotRoot() ) / CrashInfo.SourceFile;
+	}
+	else
+	{
+		FilePath = FPaths::RootDir() / CrashInfo.SourceFile;
+	}
 
 	FString Line;
-	if( FFileHelper::LoadFileToString( Line, *FilePath ) )
+	if (FFileHelper::LoadFileToString( Line, *FilePath ))
 	{
 		Line = Line.Replace( TEXT( "\r" ), TEXT( "" ) );
 		Line.ParseIntoArray( OutStrings, TEXT( "\n" ), false );
-		UE_LOG( LogCrashDebugHelper, Log, TEXT( "Reading a single source file: %s" ), *FilePath );		
+		UE_LOG( LogCrashDebugHelper, Log, TEXT( "Reading a single source file: %s" ), *FilePath );
 		return true;
 	}
 	else
diff --git a/Engine/Source/Developer/CrashDebugHelper/Private/CrashDebugPDBCache.cpp b/Engine/Source/Developer/CrashDebugHelper/Private/CrashDebugPDBCache.cpp
--- a/Engine/Source/Developer/CrashDebugHelper/Private/CrashDebugPDBCache.cpp
+++ b/Engine/Source/Developer/CrashDebugHelper/Private/CrashDebugPDBCache.cpp
@@ -16,7 +16,7 @@ void FPDBCache::Init()
 	// Default configuration
 	//PDBCachePath=F:/CrashReportPDBCache/
 	//DepotRoot=F:/depot
-	//DaysToDeleteUnusedFilesFromPDBCache=7
+	//DaysToDeleteUnusedFilesFromPDBCache=3
 	//PDBCacheSizeGB=128
 	//MinDiskFreeSpaceGB=256
 
@@ -25,6 +25,18 @@ void FPDBCache::Init()
 
 	UE_LOG( LogCrashDebugHelper, Warning, TEXT( "bUsePDBCache is %s" ), bUsePDBCache ? TEXT( "enabled" ) : TEXT( "disabled" ) );
 
+	if (bUsePDBCache)
+	{
+		GConfig->GetString( TEXT( "Engine.CrashDebugHelper" ), TEXT( "DepotRoot" ), DepotRoot, GEngineIni );
+		ICrashDebugHelper::SetDepotIndex( DepotRoot );
+
+		const bool bHasDepotRoot = IFileManager::Get().DirectoryExists( *DepotRoot );
+		UE_CLOG( !bHasDepotRoot, LogCrashDebugHelper, Warning, TEXT( "DepotRoot: %s is not valid" ), *DepotRoot );
+		UE_LOG( LogCrashDebugHelper, Log, TEXT( "DepotRoot: %s" ), *DepotRoot );	
+
+		bUsePDBCache = bHasDepotRoot;
+	}
+
 	// Get the rest of the PDB cache configuration.
 	if( bUsePDBCache )
 	{
@@ -176,7 +188,7 @@ void FPDBCache::CleanPDBCache( int32 DaysToDelete, int32 NumberOfGBsToBeCleaned
 	UE_LOG( LogCrashDebugHelper, Log, TEXT( "PDB Cache cleaned %i GBs in %.2f ms" ), NumGBsCleaned, TotalTime*1000.0f );
 }
 
-FPDBCacheEntryRef FPDBCache::CreateAndAddPDBCacheEntry( const FString& OriginalLabelName, const FString& DepotRoot, const FString& DepotName, const TArray<FString>& FilesToBeCached )
+FPDBCacheEntryRef FPDBCache::CreateAndAddPDBCacheEntry( const FString& OriginalLabelName, const FString& DepotName, const TArray<FString>& FilesToBeCached )
 {
 	const FString CleanedLabelName = EscapePath( OriginalLabelName );
 	const FString EntryDirectory = PDBCachePath / CleanedLabelName;
diff --git a/Engine/Source/Developer/CrashDebugHelper/Public/CrashDebugHelper.h b/Engine/Source/Developer/CrashDebugHelper/Public/CrashDebugHelper.h
--- a/Engine/Source/Developer/CrashDebugHelper/Public/CrashDebugHelper.h
+++ b/Engine/Source/Developer/CrashDebugHelper/Public/CrashDebugHelper.h
@@ -228,9 +228,6 @@ public:
 	static void SetDepotIndex( FString& PathToChange );
 
 protected:
-	/** Depot root, only used in conjunction with the PDB Cache. */
-	FString DepotRoot;
-
 	/**
 	 *	Pattern to search in source control for the label.
 	 *	This somehow works for older crashes, before 4.2 and for the newest one,
diff --git a/Engine/Source/Developer/CrashDebugHelper/Public/CrashDebugPDBCache.h b/Engine/Source/Developer/CrashDebugHelper/Public/CrashDebugPDBCache.h
--- a/Engine/Source/Developer/CrashDebugHelper/Public/CrashDebugPDBCache.h
+++ b/Engine/Source/Developer/CrashDebugHelper/Public/CrashDebugPDBCache.h
@@ -94,6 +94,9 @@ protected:
 	/** Path to the folder where the PDB cache is stored. */
 	FString PDBCachePath;
 
+	/** Depot root. */
+	FString DepotRoot;
+
 	/** Age of file when it should be deleted from the PDB cache. */
 	int32 DaysToDeleteUnusedFilesFromPDBCache;
 
@@ -130,6 +133,12 @@ public:
 		return bUsePDBCache;
 	}
 
+	/** @return the path the depot root. */
+	const FString& GetDepotRoot() const
+	{
+		return DepotRoot;
+	}
+
 	/** Accesses the singleton. */
 	CRASHDEBUGHELPER_API static FPDBCache& Get()
 	{
@@ -158,7 +167,7 @@ public:
 	/**
 	*	Creates a new PDB Cache entry, initializes it and adds to the database.
 	*/
-	FPDBCacheEntryRef CreateAndAddPDBCacheEntry( const FString& OriginalLabelName, const FString& DepotRoot, const FString& DepotName, const TArray<FString>& FilesToBeCached );
+	FPDBCacheEntryRef CreateAndAddPDBCacheEntry( const FString& OriginalLabelName, const FString& DepotName, const TArray<FString>& FilesToBeCached );
 
 	/**
 	*	Creates a new PDB Cache entry, initializes it and adds to the database.
