CrashDebugHelper - Symbol location also has the binaries now

diff --git a/Engine/Source/Developer/CrashDebugHelper/Private/CrashDebugHelper.cpp b/Engine/Source/Developer/CrashDebugHelper/Private/CrashDebugHelper.cpp
--- a/Engine/Source/Developer/CrashDebugHelper/Private/CrashDebugHelper.cpp
+++ b/Engine/Source/Developer/CrashDebugHelper/Private/CrashDebugHelper.cpp
@@ -193,6 +193,8 @@ bool ICrashDebugHelper::SyncModules()
 			// Find all symbols.
 			TArray<FString> NetworkSymbols;
 			IFileManager::Get().FindFilesRecursive( NetworkSymbols, *CrashInfo.SymbolsPath, TEXT( "*.pdb" ), true, false, false );
+			IFileManager::Get().FindFilesRecursive( NetworkSymbols, *CrashInfo.SymbolsPath, TEXT( "*.dll" ), true, false, false );
+			IFileManager::Get().FindFilesRecursive( NetworkSymbols, *CrashInfo.SymbolsPath, TEXT( "*.exe" ), true, false, false );
 
 			// From=Full pathname
 			// To=Relative pathname
