CrashDebugHelper - Fixed not working OLD_ENGINE version macro

diff --git a/Engine/Source/Developer/CrashDebugHelper/Private/CrashDebugHelper.cpp b/Engine/Source/Developer/CrashDebugHelper/Private/CrashDebugHelper.cpp
--- a/Engine/Source/Developer/CrashDebugHelper/Private/CrashDebugHelper.cpp
+++ b/Engine/Source/Developer/CrashDebugHelper/Private/CrashDebugHelper.cpp
@@ -700,7 +700,8 @@ void ICrashDebugHelper::FindSymbolsAndBinariesStorage()
 	
 	const FString StrENGINE_VERSION = CrashInfo.EngineVersion;
 	const FString StrPLATFORM_NAME = TEXT( "" ); // Not implemented yet
-	const FString StrOLD_ENGINE_VERSION = FString::Printf( TEXT( "++depot+%s-CL-%i" ), *CrashInfo.DepotName.Replace( P4_DEPOT_PREFIX, TEXT("") ), CrashInfo.BuiltFromCL );
+	const FString StrOLD_ENGINE_VERSION = FString::Printf( TEXT( "++depot+%s-CL-%i" ), *CrashInfo.DepotName.Replace( P4_DEPOT_PREFIX, TEXT("") ), CrashInfo.BuiltFromCL )
+		.Replace( TEXT("/"), TEXT("+") );
 
 	const FString TestExecutablesPath = ExecutablePathPattern
 		.Replace( TEXT( "%ENGINE_VERSION%" ), *StrENGINE_VERSION )
