CrashReport - UE-17082/Project debug symbols are no longer included in crash reporter callstacks

diff --git a/Engine/Source/Developer/CrashDebugHelper/Private/Windows/CrashDebugHelperWindows.cpp b/Engine/Source/Developer/CrashDebugHelper/Private/Windows/CrashDebugHelperWindows.cpp
--- a/Engine/Source/Developer/CrashDebugHelper/Private/Windows/CrashDebugHelperWindows.cpp
+++ b/Engine/Source/Developer/CrashDebugHelper/Private/Windows/CrashDebugHelperWindows.cpp
@@ -44,13 +44,13 @@ bool FCrashDebugHelperWindows::CreateMinidumpDiagnosticReport( const FString& In
 	{
 		if( CrashInfo.BuiltFromCL != FCrashInfo::INVALID_CHANGELIST )
 		{
+			// Get the build version and modules paths.
+			FCrashModuleInfo ExeFileVersion;
+			WindowsStackWalkExt.GetExeFileVersionAndModuleList( ExeFileVersion );
+
 			// CrashInfo now contains a changelist to lookup a label for
 			if( bSyncSymbols )
 			{
-				// Get the build version from the crash info.
-				FCrashModuleInfo ExeFileVersion;
-				WindowsStackWalkExt.GetExeFileVersionAndModuleList( ExeFileVersion );
-
 				FindSymbolsAndBinariesStorage();
 
 				const bool bSynced = SyncModules();
diff --git a/Engine/Source/Developer/CrashDebugHelper/Private/Windows/WindowsPlatformStackWalkExt.cpp b/Engine/Source/Developer/CrashDebugHelper/Private/Windows/WindowsPlatformStackWalkExt.cpp
--- a/Engine/Source/Developer/CrashDebugHelper/Private/Windows/WindowsPlatformStackWalkExt.cpp
+++ b/Engine/Source/Developer/CrashDebugHelper/Private/Windows/WindowsPlatformStackWalkExt.cpp
@@ -132,17 +132,18 @@ void FWindowsPlatformStackWalkExt::GetExeFileVersionAndModuleList( FCrashModuleI
 		// Get the full path of the module name
 		TCHAR ModuleName[MAX_PATH] = {0};
 		Symbol->GetModuleNameStringWide( DEBUG_MODNAME_IMAGE, ModuleIndex, ModuleBase, ModuleName, MAX_PATH, NULL );
-
-		FString RelativeModuleName = ExtractRelativePath( TEXT( "binaries" ), ModuleName );
-		if( RelativeModuleName.Len() > 0 )
+		
+		const FString RelativeModuleName = ExtractRelativePath( TEXT( "binaries" ), ModuleName );
+		// Get the exe, which we extract the version number, so we know what label to sync to
+		if (RelativeModuleName.Len() > 0 && RelativeModuleName.EndsWith( TEXT( ".exe" ) ))
 		{
-			CrashInfo.ModuleNames.Add( RelativeModuleName );
+			ExecutableIndex = ModuleIndex;
 		}
 
-		// Get the exe, which we extract the version number, so we know what label to sync to
-		if( RelativeModuleName.EndsWith( TEXT( ".exe" ) ) )
+		// Add only modules in Binaries folders
+		if (RelativeModuleName.Len() > 0)
 		{
-			ExecutableIndex = ModuleIndex;
+			CrashInfo.ModuleNames.Add( ModuleName );
 		}
 	}
 
@@ -162,100 +163,6 @@ void FWindowsPlatformStackWalkExt::GetExeFileVersionAndModuleList( FCrashModuleI
 	}
 }
 
-/**
- *	Optimized version used to grab the list of the PDB and binary files.
- *	Based on the FPackageDependencyTimestampVisitor.
- */
-class FSymbolAndExecutableFinder : public IPlatformFile::FDirectoryVisitor
-{
-private:
-	/** The file interface to use for any file operations. */
-	IPlatformFile& FileInterface;
-
-	/**
-	 *	A list of file types that should be included in the file paths.
-	 *	An extension with a dot.
-	 */
-	TArray<FString> FiletypesToInclude;
-
-	/** A list of directories that we should not traverse into. */
-	TArray<FString> DirectoriesToIgnore;
-
-	/** Relative paths to local files. */
-	TArray<FString> FilePaths;
-
-public:
-
-	/** A list of directories to be checked, to avoid using recursive algorithm and save stack-memory. */
-	TArray<FString> RecursiveDirectories;
-
-	FSymbolAndExecutableFinder( IPlatformFile& InFileInterface, const TArray<FString>& InFiletypesToInclude, const TArray<FString>& InDirectoriesToIgnore )
-		: FileInterface( InFileInterface )
-		, FiletypesToInclude( InFiletypesToInclude )
-		, DirectoriesToIgnore( InDirectoriesToIgnore )
-	{}
-
-	virtual bool Visit( const TCHAR* FilenameOrDirectory, bool bIsDirectory )
-	{
-		// Make sure all paths are "standardized" so the other end can match up with it's own standardized paths
-		FString RelativeFilename = FilenameOrDirectory;
-		FPaths::MakeStandardFilename( RelativeFilename );
-
-		// Grab the file.
-		if( !bIsDirectory )
-		{
-			if( FiletypesToInclude.Num() > 0 )
-			{
-				// If it is a file that matches the filetypes, add it.
-				for( const auto& FileExtension : FiletypesToInclude )
-				{
-					if( RelativeFilename.Contains( FileExtension ) )
-					{
-						new(FilePaths) FString( RelativeFilename );
-						break;
-					}
-				}
-			}
-			else
-			{
-				new(FilePaths) FString( RelativeFilename );
-			}
-		}
-
-		// Iterate over directories we care about.
-		if( bIsDirectory )
-		{
-			bool bShouldRecurse = true;
-			// Look in all the ignore directories looking for a match.
-			for( const auto& Directory : DirectoriesToIgnore )
-			{
-				if( RelativeFilename.EndsWith( Directory ) )
-				{
-					bShouldRecurse = false;
-					break;
-				}
-			}
-
-			// Recurse if we should.
-			if( bShouldRecurse )
-			{
-				RecursiveDirectories.Add(FilenameOrDirectory);
-			}
-		}
-
-		return true;
-	}
-
-	
-	/**
-	 * @return a reference to the results.
-	 */
-	const TArray<FString>& GetFilePaths() const
-	{
-		return FilePaths;
-	}
-};
-
 void FWindowsPlatformStackWalkExt::SetSymbolPathsFromModules()
 {
 	const bool bUseCachedData = CrashInfo.PDBCacheEntry.IsValid();
@@ -287,44 +194,9 @@ void FWindowsPlatformStackWalkExt::SetSymbolPathsFromModules()
 	// For locally launched minidump diagnostics.
 	else
 	{
-		// We care only about the PDB and binary files.
-		TArray<FString> FiletypesToInclude;
-		FiletypesToInclude.Add( TEXT( ".pdb" ) );
-		FiletypesToInclude.Add( TEXT( ".dll" ) );
-		FiletypesToInclude.Add( TEXT( ".exe" ) );
-		// These folders don't contain any binary files, so skip them to improve the finder performance.
-		TArray<FString> DirectoriesToIgnore;
-		DirectoriesToIgnore.Add( TEXT( "/Content" ) );
-		DirectoriesToIgnore.Add( TEXT( "/DerivedDataCache" ) );
-		DirectoriesToIgnore.Add( TEXT( "/Intermediate" ) );
-		DirectoriesToIgnore.Add( TEXT( "/Documentation" ) );
-		DirectoriesToIgnore.Add( TEXT( "/NoRedist" ) );
-		DirectoriesToIgnore.Add( TEXT( "/Public" ) );
-		DirectoriesToIgnore.Add( TEXT( "/Private" ) );
-
-		// Platforms to ignore.
-		DirectoriesToIgnore.Add( TEXT( "/Android" ) );
-		DirectoriesToIgnore.Add( TEXT( "/Cygwin" ) );
-		DirectoriesToIgnore.Add( TEXT( "/HTML5" ) );
-		DirectoriesToIgnore.Add( TEXT( "/IOS" ) );
-		DirectoriesToIgnore.Add( TEXT( "/Java" ) );
-		DirectoriesToIgnore.Add( TEXT( "/Linux" ) );
-		DirectoriesToIgnore.Add( TEXT( "/Mac" ) );
-		DirectoriesToIgnore.Add( TEXT( "/PS4" ) );
-		DirectoriesToIgnore.Add( TEXT( "/XboxOne" ) );
-
-		IPlatformFile& PlatformFile = FPlatformFileManager::Get().GetPlatformFile();
-		FSymbolAndExecutableFinder SymbolAndExecutableFinder( PlatformFile, FiletypesToInclude, DirectoriesToIgnore );
-		SymbolAndExecutableFinder.Visit( *FPaths::RootDir(), true );
-
-		while( SymbolAndExecutableFinder.RecursiveDirectories.Num() )
-		{
-			const FString Directory = SymbolAndExecutableFinder.RecursiveDirectories.Pop();
-			PlatformFile.IterateDirectory( *Directory, SymbolAndExecutableFinder );
-		}
-
+		// Iterate over all loaded modules.
 		TSet<FString> SymbolPaths;
-		for( const auto& Filename : SymbolAndExecutableFinder.GetFilePaths() )
+		for (const auto& Filename : CrashInfo.ModuleNames)
 		{
 			const FString Path = FPaths::GetPath( Filename );
 			if( Path.Len() > 0 )
