CrashReport - Better callstacks for crashes in runnable threads

diff --git a/Engine/Source/Developer/CrashDebugHelper/Private/Windows/WindowsPlatformStackWalkExt.cpp b/Engine/Source/Developer/CrashDebugHelper/Private/Windows/WindowsPlatformStackWalkExt.cpp
--- a/Engine/Source/Developer/CrashDebugHelper/Private/Windows/WindowsPlatformStackWalkExt.cpp
+++ b/Engine/Source/Developer/CrashDebugHelper/Private/Windows/WindowsPlatformStackWalkExt.cpp
@@ -452,7 +452,7 @@ int FWindowsPlatformStackWalkExt::GetCallstacks()
 	FMemory::Memzero( ContextData, MaxFramesSize );
 	HRESULT HR = Control->GetContextStackTrace( Context, ContextUsed, StackFrames, MaxFrames, ContextData, MaxFramesSize, ContextUsed, &Count );
 
-	int32 AssertOnEnsureIndex = -1;
+	int32 AssertOrEnsureIndex = -1;
 
 	for( uint32 StackIndex = 0; StackIndex < Count; StackIndex++ )
 	{	
@@ -467,7 +467,7 @@ int FWindowsPlatformStackWalkExt::GetCallstacks()
 			FString ModuleAndFunction = NameByOffset;
 
 			// Don't care about any more entries higher than this
-			if( ModuleAndFunction.Contains( TEXT( "tmainCRTStartup" ) ) )
+			if (ModuleAndFunction.Contains( TEXT( "tmainCRTStartup" ) ) || ModuleAndFunction.Contains( TEXT( "FRunnableThreadWin::GuardedRun" ) ))
 			{
 				break;
 			}
@@ -509,7 +509,7 @@ int FWindowsPlatformStackWalkExt::GetCallstacks()
 					|| FunctionName.Contains( TEXT( "NewReportEnsure" ), ESearchCase::CaseSensitive ) )
 				{
 					bFoundSourceFile = false;
-					AssertOnEnsureIndex = FMath::Max( AssertOnEnsureIndex, (int32)StackIndex );
+					AssertOrEnsureIndex = FMath::Max( AssertOrEnsureIndex, (int32)StackIndex );
 				}
 			}
 
@@ -524,9 +524,9 @@ int FWindowsPlatformStackWalkExt::GetCallstacks()
 	}
 
 	// Remove callstack entries below FDebug, we don't need them.
-	if (AssertOnEnsureIndex > 0)
+	if (AssertOrEnsureIndex > 0)
 	{	
-		Exception.CallStackString.RemoveAt( 0, AssertOnEnsureIndex );
+		Exception.CallStackString.RemoveAt( 0, AssertOrEnsureIndex );
 		UE_LOG( LogCrashDebugHelper, Warning, TEXT( "Callstack trimmed to %i entries" ), Exception.CallStackString.Num() );
 	}
 
