Implement Source Context reporting for Mac CrashReportClient as requested in UE-16101.
- This CL also exports the PLCrashReporter dump from binary format to text so that it is legible without further processing.

diff --git a/Engine/Source/Developer/CrashDebugHelper/Private/Mac/CrashDebugHelperMac.cpp b/Engine/Source/Developer/CrashDebugHelper/Private/Mac/CrashDebugHelperMac.cpp
--- a/Engine/Source/Developer/CrashDebugHelper/Private/Mac/CrashDebugHelperMac.cpp
+++ b/Engine/Source/Developer/CrashDebugHelper/Private/Mac/CrashDebugHelperMac.cpp
@@ -7,6 +7,37 @@
 #include "CrashReporter.h"
 #include <cxxabi.h>
 
+FString ExtractRelativePath( const TCHAR* BaseName, TCHAR const* FullName )
+{
+	FString FullPath = FString( FullName ).ToLower();
+	FullPath = FullPath.Replace( TEXT( "\\" ), TEXT( "/" ) );
+	
+	TArray<FString> Components;
+	int32 Count = FullPath.ParseIntoArray( Components, TEXT( "/" ), true );
+	FullPath = TEXT( "" );
+	
+	for( int32 Index = 0; Index < Count; Index++ )
+	{
+		if( Components[Index] == BaseName )
+		{
+			if( Index > 0 )
+			{
+				for( int32 Inner = Index - 1; Inner < Count; Inner++ )
+				{
+					FullPath += Components[Inner];
+					if( Inner < Count - 1 )
+					{
+						FullPath += TEXT( "/" );
+					}
+				}
+			}
+			break;
+		}
+	}
+	
+	return FullPath;
+}
+
 static int32 ParseReportVersion(TCHAR const* CrashLog, int32& OutVersion)
 {
 	int32 Found = 0;
@@ -511,6 +542,9 @@ bool FCrashDebugHelperMac::ParseCrashDump(const FString& InCrashDumpName, FCrash
 		{
 			NSString* Report = [PLCrashReportTextFormatter stringValueForCrashReport: CrashLog withTextFormat: PLCrashReportTextFormatiOS];
 			CrashDump = FString(Report);
+			
+			// Replace the binary PLCrashReporter file with an easily readable text dump
+			FFileHelper::SaveStringToFile(CrashDump, *InCrashDumpName);
 		}
 	}
 	
@@ -554,6 +588,14 @@ bool FCrashDebugHelperMac::ParseCrashDump(const FString& InCrashDumpName, FCrash
 bool FCrashDebugHelperMac::CreateMinidumpDiagnosticReport( const FString& InCrashDumpName )
 {
 	bool bOK = false;
+	const bool bSyncSymbols = FParse::Param( FCommandLine::Get(), TEXT( "SyncSymbols" ) );
+	const bool bAnnotate = FParse::Param( FCommandLine::Get(), TEXT( "Annotate" ) );
+	const bool bUseSCC = bSyncSymbols || bAnnotate;
+	
+	if( bUseSCC )
+	{
+		InitSourceControl( false );
+	}
 	
 	FString CrashDump;
 	
@@ -567,7 +609,10 @@ bool FCrashDebugHelperMac::CreateMinidumpDiagnosticReport( const FString& InCras
             if(CrashLog && !Error)
             {
                 NSString* Report = [PLCrashReportTextFormatter stringValueForCrashReport: CrashLog withTextFormat: PLCrashReportTextFormatiOS];
-                CrashDump = FString(Report);
+				CrashDump = FString(Report);
+				
+				// Replace the binary PLCrashReporter file with an easily readable text dump
+				FFileHelper::SaveStringToFile(CrashDump, *InCrashDumpName);
             }
         }
     }
@@ -608,6 +653,13 @@ bool FCrashDebugHelperMac::CreateMinidumpDiagnosticReport( const FString& InCras
 			if(Result == 5 && Branch.Len() > 0)
 			{
 				CrashInfo.LabelName = Branch;
+				
+				if( bSyncSymbols )
+				{
+					FindSymbolsAndBinariesStorage();
+					
+					SyncModules();
+				}
 			}
 			
 			Result = ParseOS(*CrashDump, CrashInfo.SystemInfo.OSMajor, CrashInfo.SystemInfo.OSMinor, CrashInfo.SystemInfo.OSBuild, CrashInfo.SystemInfo.OSRevision);
@@ -658,6 +710,7 @@ bool FCrashDebugHelperMac::CreateMinidumpDiagnosticReport( const FString& InCras
 			
 			FPlatformSymbolDatabaseSet SymbolCache;
 			
+			bool bIsCrashLocation = true;
 			TCHAR const* ThreadStackLine = FindCrashedThreadStack(*CrashDump);
 			while(ThreadStackLine)
 			{
@@ -694,6 +747,35 @@ bool FCrashDebugHelperMac::CreateMinidumpDiagnosticReport( const FString& InCras
 						
 					case 5:
 					case 6: // Function name might be parsed twice
+						if(bIsCrashLocation)
+						{
+							if( FileName.Len() > 0 && LineNumber > 0 )
+							{
+								// Sync the source file where the crash occurred
+								CrashInfo.SourceFile = ExtractRelativePath( TEXT( "source" ), *FileName );
+								CrashInfo.SourceLineNumber = LineNumber;
+								
+								if( bSyncSymbols && CrashInfo.BuiltFromCL > 0 )
+								{
+									UE_LOG( LogCrashDebugHelper, Log, TEXT( "Using CL %i to sync crash source file" ), CrashInfo.BuiltFromCL );
+									SyncSourceFile();
+								}
+								
+								// Try to annotate the file if requested
+								bool bAnnotationSuccessful = false;
+								if( bAnnotate )
+								{
+									bAnnotationSuccessful = AddAnnotatedSourceToReport();
+								}
+								
+								// If annotation is not requested, or failed, add the standard source context
+								if( !bAnnotationSuccessful )
+								{
+									AddSourceToReport();
+								}
+							}
+						}
+						
 						CrashInfo.Exception.CallStackString.Push( FString::Printf( TEXT( "%s Address = 0x%lx [%s, line %d] [in %s]" ), *FunctionName, ProgramCounter, *FileName, LineNumber, *ModuleName ) );
 						ThreadInfo.CallStack.Push(ProgramCounter);
 						
@@ -706,6 +788,8 @@ bool FCrashDebugHelperMac::CreateMinidumpDiagnosticReport( const FString& InCras
 						ThreadStackLine = nullptr;
 						break;
 				}
+				
+				bIsCrashLocation = false;
 			}
 			
 			CrashInfo.Threads.Push(ThreadInfo);
@@ -714,5 +798,10 @@ bool FCrashDebugHelperMac::CreateMinidumpDiagnosticReport( const FString& InCras
 		}
 	}
 	
+	if( bUseSCC )
+	{
+		ShutdownSourceControl();
+	}
+	
 	return bOK;
 }
