CrashReport - Crash debug helper code clean-up

diff --git a/Engine/Source/Developer/CrashDebugHelper/Private/CrashDebugHelper.cpp b/Engine/Source/Developer/CrashDebugHelper/Private/CrashDebugHelper.cpp
--- a/Engine/Source/Developer/CrashDebugHelper/Private/CrashDebugHelper.cpp
+++ b/Engine/Source/Developer/CrashDebugHelper/Private/CrashDebugHelper.cpp
@@ -642,20 +642,6 @@ void FCrashInfo::GenerateReport( const FString& DiagnosticsPath )
 	}
 }
 
-bool ICrashDebugHelper::SyncRequiredFilesForDebuggingFromLabel(const FString& InLabel, const FString& InPlatform)
-{
-	// @TODO yrx 2015-02-19 Use PDB cache
-
-	return false;
-}
-
-bool ICrashDebugHelper::SyncRequiredFilesForDebuggingFromChangelist(int32 InChangelistNumber, const FString& InPlatform)
-{
-	// @TODO yrx 2015-02-19 Use PDB cache
-
-	return false;
-}
-
 void ICrashDebugHelper::FindSymbolsAndBinariesStorage()
 {
 	CrashInfo.ExecutablesPath.Empty();
diff --git a/Engine/Source/Developer/CrashDebugHelper/Private/CrashDebugHelperModule.cpp b/Engine/Source/Developer/CrashDebugHelper/Private/CrashDebugHelperModule.cpp
--- a/Engine/Source/Developer/CrashDebugHelper/Private/CrashDebugHelperModule.cpp
+++ b/Engine/Source/Developer/CrashDebugHelper/Private/CrashDebugHelperModule.cpp
@@ -21,7 +21,7 @@ DEFINE_LOG_CATEGORY(LogCrashDebugHelper);
 void FCrashDebugHelperModule::StartupModule()
 {
 	CrashDebugHelper = new FCrashDebugHelper();
-	if (CrashDebugHelper != NULL)
+	if (CrashDebugHelper != nullptr)
 	{
 		CrashDebugHelper->Init();
 	}
@@ -29,10 +29,10 @@ void FCrashDebugHelperModule::StartupModule()
 
 void FCrashDebugHelperModule::ShutdownModule()
 {
-	if (CrashDebugHelper != NULL)
+	if (CrashDebugHelper != nullptr)
 	{
 		delete CrashDebugHelper;
-		CrashDebugHelper = NULL;
+		CrashDebugHelper = nullptr;
 	}
 }
 
diff --git a/Engine/Source/Developer/CrashDebugHelper/Private/Linux/CrashDebugHelperLinux.cpp b/Engine/Source/Developer/CrashDebugHelper/Private/Linux/CrashDebugHelperLinux.cpp
--- a/Engine/Source/Developer/CrashDebugHelper/Private/Linux/CrashDebugHelperLinux.cpp
+++ b/Engine/Source/Developer/CrashDebugHelper/Private/Linux/CrashDebugHelperLinux.cpp
@@ -3,19 +3,6 @@
 #include "CrashDebugHelperPrivatePCH.h"
 #include "CrashDebugHelperLinux.h"
 
-FCrashDebugHelperLinux::FCrashDebugHelperLinux()
-{
-}
-
-FCrashDebugHelperLinux::~FCrashDebugHelperLinux()
-{
-}
-
-bool FCrashDebugHelperLinux::ParseCrashDump(const FString& InCrashDumpName, FCrashDebugInfo& OutCrashDebugInfo)
-{
-	return false;
-}
-
 bool FCrashDebugHelperLinux::CreateMinidumpDiagnosticReport( const FString& InCrashDumpName )
 {
 	return true;
diff --git a/Engine/Source/Developer/CrashDebugHelper/Private/Linux/CrashDebugHelperLinux.h b/Engine/Source/Developer/CrashDebugHelper/Private/Linux/CrashDebugHelperLinux.h
--- a/Engine/Source/Developer/CrashDebugHelper/Private/Linux/CrashDebugHelperLinux.h
+++ b/Engine/Source/Developer/CrashDebugHelper/Private/Linux/CrashDebugHelperLinux.h
@@ -7,19 +7,6 @@
 class FCrashDebugHelperLinux : public ICrashDebugHelper
 {
 public:
-	FCrashDebugHelperLinux();
-	virtual ~FCrashDebugHelperLinux();
-
-	/**
-	 *	Parse the given crash dump, determining EngineVersion of the build that produced it - if possible. 
-	 *
-	 *	@param	InCrashDumpName		The crash dump file to process
-	 *	@param	OutCrashDebugInfo	The crash dump info extracted from the file
-	 *
-	 *	@return	bool				true if successful, false if not
-	 */
-	virtual bool ParseCrashDump(const FString& InCrashDumpName, FCrashDebugInfo& OutCrashDebugInfo) override;
-
 	/**
 	 *	Parse the given crash dump, and generate a report. 
 	 *
diff --git a/Engine/Source/Developer/CrashDebugHelper/Private/Windows/CrashDebugHelperWindows.cpp b/Engine/Source/Developer/CrashDebugHelper/Private/Windows/CrashDebugHelperWindows.cpp
--- a/Engine/Source/Developer/CrashDebugHelper/Private/Windows/CrashDebugHelperWindows.cpp
+++ b/Engine/Source/Developer/CrashDebugHelper/Private/Windows/CrashDebugHelperWindows.cpp
@@ -10,20 +10,6 @@
 #include "AllowWindowsPlatformTypes.h"
 #include <DbgHelp.h>
 
-FCrashDebugHelperWindows::FCrashDebugHelperWindows()
-{
-}
-
-FCrashDebugHelperWindows::~FCrashDebugHelperWindows()
-{
-}
-
-
-bool FCrashDebugHelperWindows::ParseCrashDump(const FString& InCrashDumpName, FCrashDebugInfo& OutCrashDebugInfo)
-{
-	return true;
-}
-
 bool FCrashDebugHelperWindows::CreateMinidumpDiagnosticReport( const FString& InCrashDumpFilename )
 {
 	const bool bSyncSymbols = FParse::Param( FCommandLine::Get(), TEXT( "SyncSymbols" ) );
diff --git a/Engine/Source/Developer/CrashDebugHelper/Private/Windows/CrashDebugHelperWindows.h b/Engine/Source/Developer/CrashDebugHelper/Private/Windows/CrashDebugHelperWindows.h
--- a/Engine/Source/Developer/CrashDebugHelper/Private/Windows/CrashDebugHelperWindows.h
+++ b/Engine/Source/Developer/CrashDebugHelper/Private/Windows/CrashDebugHelperWindows.h
@@ -7,19 +7,6 @@
 class FCrashDebugHelperWindows : public ICrashDebugHelper
 {
 public:
-	FCrashDebugHelperWindows();
-	virtual ~FCrashDebugHelperWindows();
-
-	/**
-	 *	Parse the given crash dump, determining EngineVersion of the build that produced it - if possible. 
-	 *
-	 *	@param	InCrashDumpName		The crash dump file to process
-	 *	@param	OutCrashDebugInfo	The crash dump info extracted from the file
-	 *
-	 *	@return	bool				true if successful, false if not
-	 */
-	virtual bool ParseCrashDump(const FString& InCrashDumpName, FCrashDebugInfo& OutCrashDebugInfo) override;
-
 	/**
 	 *	Parse the given crash dump, and generate a report. 
 	 *
diff --git a/Engine/Source/Developer/CrashDebugHelper/Private/Windows/WindowsPlatformStackWalkExt.cpp b/Engine/Source/Developer/CrashDebugHelper/Private/Windows/WindowsPlatformStackWalkExt.cpp
--- a/Engine/Source/Developer/CrashDebugHelper/Private/Windows/WindowsPlatformStackWalkExt.cpp
+++ b/Engine/Source/Developer/CrashDebugHelper/Private/Windows/WindowsPlatformStackWalkExt.cpp
@@ -501,7 +501,7 @@ int FWindowsPlatformStackWalkExt::GetCallstacks()
 				ModuleName = ModuleAndFunction;
 			}
 
-			// @TODO yrx 2015-02-24 Add for other platforms
+			// #YRX_Crash: 2015-07-24 Add for other platforms
 			// If we find an assert, the actual source file we're interested in is the next one up, so reset the source file found flag
 			if( FunctionName.Len() > 0 )
 			{
diff --git a/Engine/Source/Developer/CrashDebugHelper/Public/CrashDebugHelper.h b/Engine/Source/Developer/CrashDebugHelper/Public/CrashDebugHelper.h
--- a/Engine/Source/Developer/CrashDebugHelper/Public/CrashDebugHelper.h
+++ b/Engine/Source/Developer/CrashDebugHelper/Public/CrashDebugHelper.h
@@ -120,7 +120,7 @@ public:
 	}
 };
 
-// @TODO yrx 2015-02-18 Merge with FCrashDescription
+// #YRX_Crash: 2015-07-24 Merge with FCrashDescription
 /** A platform independent representation of a crash */
 class CRASHDEBUGHELPER_API FCrashInfo
 {
@@ -271,6 +271,8 @@ public:
 	 *	@param	OutCrashDebugInfo	The crash dump info extracted from the file
 	 *
 	 *	@return	bool				true if successful, false if not
+	 *	
+	 *	Only used by Mac, to be removed.
 	 */
 	virtual bool ParseCrashDump(const FString& InCrashDumpName, FCrashDebugInfo& OutCrashDebugInfo)
 	{
@@ -289,28 +291,6 @@ public:
 		return false;
 	}
 
-	/**
-	 *	Sync the required files from source control for debugging
-	 *
-	 *	@param	InLabel			The label to sync files from
-	 *	@param	InPlatform		The platform to sync files for
-	 *
-	 *	@return bool			true if successful, false if not
-	 */
-	virtual bool SyncRequiredFilesForDebuggingFromLabel(const FString& InLabel, const FString& InPlatform);
-
-	/**
-	 *	Sync the required files from source control for debugging
-	 *	NOTE: Currently this will only handle changelists of 'valid' builds from the build machine
-	 *			Ie the changelist must map to a build label.
-	 *
-	 *	@param	InChangelist	The changelist to sync files from
-	 *	@param	InPlatform		The platform to sync files for
-	 *
-	 *	@return bool			true if successful, false if not
-	 */
-	virtual bool SyncRequiredFilesForDebuggingFromChangelist(int32 InChangelistNumber, const FString& InPlatform);
-
 	/**
 	 * Sync the branch root relative file names to the requested label
 	 */
diff --git a/Engine/Source/Developer/CrashDebugHelper/Public/CrashDebugHelperModule.h b/Engine/Source/Developer/CrashDebugHelper/Public/CrashDebugHelperModule.h
--- a/Engine/Source/Developer/CrashDebugHelper/Public/CrashDebugHelperModule.h
+++ b/Engine/Source/Developer/CrashDebugHelper/Public/CrashDebugHelperModule.h
@@ -13,7 +13,7 @@ public:
 	virtual void ShutdownModule();
 
 	/** Gets the crash debug helper singleton or returns NULL */
-	CRASHDEBUGHELPER_API virtual ICrashDebugHelper* Get();
+	CRASHDEBUGHELPER_API ICrashDebugHelper* Get();
 
 private:
 	class ICrashDebugHelper* CrashDebugHelper;
