#UE4 Softer error handling for missing timestamps in PDBCache folders

diff --git a/Engine/Source/Developer/CrashDebugHelper/Private/CrashDebugPDBCache.cpp b/Engine/Source/Developer/CrashDebugHelper/Private/CrashDebugPDBCache.cpp
--- a/Engine/Source/Developer/CrashDebugHelper/Private/CrashDebugPDBCache.cpp
+++ b/Engine/Source/Developer/CrashDebugHelper/Private/CrashDebugPDBCache.cpp
@@ -121,8 +121,11 @@ void FPDBCache::InitializePDBCache()
 
 	for( const auto& Directory : PDBCacheEntryDirectories )
 	{
-		FPDBCacheEntryRef Entry = ReadPDBCacheEntry( Directory );
-		PDBCacheEntries.Add( Directory, Entry );
+		FPDBCacheEntryPtr Entry = ReadPDBCacheEntry( Directory );
+		if (Entry.IsValid())
+		{
+			PDBCacheEntries.Add( Directory, Entry.ToSharedRef() );
+		}
 	}
 
 	SortPDBCache();
@@ -279,7 +282,7 @@ FPDBCacheEntryRef FPDBCache::CreateAndAddPDBCacheEntryMixed( const FString& Prod
 	return NewCacheEntry;
 }
 
-FPDBCacheEntryRef FPDBCache::ReadPDBCacheEntry( const FString& Directory )
+FPDBCacheEntryPtr FPDBCache::ReadPDBCacheEntry( const FString& Directory )
 {
 	const FString EntryDirectory = PDBCachePath / Directory;
 	const FString EntryTimeStampFilenameNoMeta = EntryDirectory / PDBTimeStampFileNoMeta;
@@ -325,10 +328,10 @@ FPDBCacheEntryRef FPDBCache::ReadPDBCacheEntry( const FString& Directory )
 	else
 	{
 		// Something wrong.
-		checkf( 0, TEXT( "Invalid symbol cache entry: %s" ), *EntryDirectory );
+		ensureMsgf( 0, TEXT( "Invalid symbol cache entry: %s" ), *EntryDirectory );
 	}
 
-	return NewEntry.ToSharedRef();
+	return NewEntry;
 }
 
 void FPDBCache::TouchPDBCacheEntry( const FString& Directory )
diff --git a/Engine/Source/Developer/CrashDebugHelper/Public/CrashDebugPDBCache.h b/Engine/Source/Developer/CrashDebugHelper/Public/CrashDebugPDBCache.h
--- a/Engine/Source/Developer/CrashDebugHelper/Public/CrashDebugPDBCache.h
+++ b/Engine/Source/Developer/CrashDebugHelper/Public/CrashDebugPDBCache.h
@@ -216,7 +216,7 @@ protected:
 	/**
 	*	Reads an existing PDB Cache entry.
 	*/
-	FPDBCacheEntryRef ReadPDBCacheEntry( const FString& Directory );
+	FPDBCacheEntryPtr ReadPDBCacheEntry( const FString& Directory );
 
 	/**
 	*	Sort PDB Cache entries by last access time.
