Crash Reporter doesn't produce callstack

diff --git a/Engine/Source/Developer/CrashDebugHelper/Private/CrashDebugHelper.cpp b/Engine/Source/Developer/CrashDebugHelper/Private/CrashDebugHelper.cpp
--- a/Engine/Source/Developer/CrashDebugHelper/Private/CrashDebugHelper.cpp
+++ b/Engine/Source/Developer/CrashDebugHelper/Private/CrashDebugHelper.cpp
@@ -159,10 +159,13 @@ bool ICrashDebugHelper::Init()
 	}
 	else
 	{
-		// Obsolete.
+		// Use the current values.
+		const FEngineVersion& EngineVersion = FEngineVersion::Current();
+		CrashInfo.DepotName = EngineVersion.GetBranch();
+		CrashInfo.BuiltFromCL = (int32)EngineVersion.GetChangelist();
+		CrashInfo.EngineVersion = EngineVersion.ToString();
 	}
 
-
 	UE_LOG( LogCrashDebugHelper, Log, TEXT( "DepotName: %s" ), *CrashInfo.DepotName );
 	UE_LOG( LogCrashDebugHelper, Log, TEXT( "BuiltFromCL: %i" ), CrashInfo.BuiltFromCL );
 	UE_LOG( LogCrashDebugHelper, Log, TEXT( "EngineVersion: %s" ), *CrashInfo.EngineVersion );
