Attempt to fix "Branch not found" errors on MDD on CR server.
Code change to CrashDebugHelper.
Updated MDD binaries.

diff --git a/Engine/Source/Developer/CrashDebugHelper/Private/CrashDebugHelper.cpp b/Engine/Source/Developer/CrashDebugHelper/Private/CrashDebugHelper.cpp
--- a/Engine/Source/Developer/CrashDebugHelper/Private/CrashDebugHelper.cpp
+++ b/Engine/Source/Developer/CrashDebugHelper/Private/CrashDebugHelper.cpp
@@ -770,7 +770,7 @@ void ICrashDebugHelper::FindSymbolsAndBinariesStorage()
 	FString ExecutablePathPattern = ConfigEntry.ExecutablePathPattern;
 	FString SymbolPathPattern = ConfigEntry.ExecutablePathPattern;
 
-	if (!ExecutablePathPattern.IsEmpty() && !SymbolPathPattern.IsEmpty())
+	if (!ExecutablePathPattern.IsEmpty() || !SymbolPathPattern.IsEmpty())
 	{
 		UE_LOG( LogCrashDebugHelper, Log, TEXT( "Using branch: %s" ), *CrashInfo.DepotName );
 	}
@@ -799,8 +799,16 @@ void ICrashDebugHelper::FindSymbolsAndBinariesStorage()
 	// If this step successes, we will grab the executable from the network path instead of P4.
 	bool bFoundDirectory = false;
 	
-	const bool bHasExecutables = IFileManager::Get().DirectoryExists( *TestExecutablesPath );
-	const bool bHasSymbols = IFileManager::Get().DirectoryExists( *TestSymbolsPath );
+	bool bHasExecutables = false;
+	if (!TestExecutablesPath.IsEmpty())
+	{
+		bHasExecutables = IFileManager::Get().DirectoryExists(*TestExecutablesPath);
+	}
+	bool bHasSymbols = false;
+	if (!TestSymbolsPath.IsEmpty())
+	{
+		bHasSymbols = IFileManager::Get().DirectoryExists(*TestSymbolsPath);
+	}
 
 	if( bHasExecutables && bHasSymbols )
 	{
