Added new command line args to MDD to allow the caller to pass platform names.
[OPP-4792] - Crash reports not appearing for launcher on Windows 32-bit OS

Implemented platform name substitution strings in exe/symbol search patterns (patterns not using them yet).
Simplified pdb cache defines in config to cover more branches per product and remove platforms.
Code change to CrashDebugHelper.
Engine config change.
Updated MDD binaries.

diff --git a/Engine/Source/Developer/CrashDebugHelper/Private/CrashDebugHelper.cpp b/Engine/Source/Developer/CrashDebugHelper/Private/CrashDebugHelper.cpp
--- a/Engine/Source/Developer/CrashDebugHelper/Private/CrashDebugHelper.cpp
+++ b/Engine/Source/Developer/CrashDebugHelper/Private/CrashDebugHelper.cpp
@@ -166,6 +166,30 @@ bool ICrashDebugHelper::Init()
 		CrashInfo.EngineVersion = EngineVersion.ToString();
 	}
 
+	FString PlatformName;
+	const bool bHasPlatformName = FParse::Value(FCommandLine::Get(), TEXT("PlatformName="), PlatformName);
+	if (bHasPlatformName)
+	{
+		CrashInfo.PlatformName = PlatformName;
+	}
+	else
+	{
+		// Use the current values.
+		CrashInfo.PlatformName = FPlatformProperties::PlatformName();
+	}
+
+	FString PlatformVariantName;
+	const bool bHasPlatformVariantName = FParse::Value(FCommandLine::Get(), TEXT("PlatformVariantName="), PlatformVariantName);
+	if (bHasPlatformVariantName)
+	{
+		CrashInfo.PlatformVariantName = PlatformVariantName;
+	}
+	else
+	{
+		// Use the basic platform name.
+		CrashInfo.PlatformVariantName = CrashInfo.PlatformName;
+	}
+
 	UE_LOG( LogCrashDebugHelper, Log, TEXT( "DepotName: %s" ), *CrashInfo.DepotName );
 	UE_LOG( LogCrashDebugHelper, Log, TEXT( "BuiltFromCL: %i" ), CrashInfo.BuiltFromCL );
 	UE_LOG( LogCrashDebugHelper, Log, TEXT( "EngineVersion: %s" ), *CrashInfo.EngineVersion );
@@ -781,18 +805,21 @@ void ICrashDebugHelper::FindSymbolsAndBinariesStorage()
 	}
 	
 	const FString StrENGINE_VERSION = CrashInfo.EngineVersion;
-	const FString StrPLATFORM_NAME = TEXT( "" ); // Not implemented yet
+	const FString StrPLATFORM_NAME = CrashInfo.PlatformName;
+	const FString StrPLATFORM_VARIANT = CrashInfo.PlatformVariantName;
 	const FString StrOLD_ENGINE_VERSION = FString::Printf( TEXT( "%s-CL-%i" ), *CrashInfo.DepotName.Replace( TEXT( "+" ), TEXT( "/" ) ), CrashInfo.BuiltFromCL )
 		.Replace( TEXT("/"), TEXT("+") );
 
 	const FString TestExecutablesPath = ExecutablePathPattern
 		.Replace( TEXT( "%ENGINE_VERSION%" ), *StrENGINE_VERSION )
 		.Replace( TEXT( "%PLATFORM_NAME%" ), *StrPLATFORM_NAME )
+		.Replace( TEXT( "%PLATFORM_VARIANT%" ), *StrPLATFORM_VARIANT )
 		.Replace( TEXT( "%OLD_ENGINE_VERSION%" ), *StrOLD_ENGINE_VERSION );
 
 	const FString TestSymbolsPath = SymbolPathPattern
 		.Replace( TEXT( "%ENGINE_VERSION%" ), *StrENGINE_VERSION )
 		.Replace( TEXT( "%PLATFORM_NAME%" ), *StrPLATFORM_NAME )
+		.Replace( TEXT( "%PLATFORM_VARIANT%" ), *StrPLATFORM_VARIANT )
 		.Replace( TEXT( "%OLD_ENGINE_VERSION%" ), *StrOLD_ENGINE_VERSION );
 
 	// Try to find the network path by using the pattern supplied via ini.
diff --git a/Engine/Source/Developer/CrashDebugHelper/Private/CrashDebugPDBCache.cpp b/Engine/Source/Developer/CrashDebugHelper/Private/CrashDebugPDBCache.cpp
--- a/Engine/Source/Developer/CrashDebugHelper/Private/CrashDebugPDBCache.cpp
+++ b/Engine/Source/Developer/CrashDebugHelper/Private/CrashDebugPDBCache.cpp
@@ -17,7 +17,7 @@ void FPDBCache::Init()
 	//PDBCachePath=F:/CrashReportPDBCache/
 	//DepotRoot=F:/depot
 	//DaysToDeleteUnusedFilesFromPDBCache=3
-	//PDBCacheSizeGB=128
+	//PDBCacheSizeGB=300
 	//MinDiskFreeSpaceGB=256
 
 	// Can be enabled only through the command line.
diff --git a/Engine/Source/Developer/CrashDebugHelper/Public/CrashDebugHelper.h b/Engine/Source/Developer/CrashDebugHelper/Public/CrashDebugHelper.h
--- a/Engine/Source/Developer/CrashDebugHelper/Public/CrashDebugHelper.h
+++ b/Engine/Source/Developer/CrashDebugHelper/Public/CrashDebugHelper.h
@@ -167,6 +167,9 @@ public:
 	/** Shared pointer to the PDB Cache entry, if valid contains all information about synced PDBs. */
 	TSharedPtr<FPDBCacheEntry> PDBCacheEntry;
 
+	FString PlatformName;
+	FString PlatformVariantName;
+
 	FCrashInfo()
 		: BuiltFromCL( INVALID_CHANGELIST )
 		, SourceLineNumber( 0 )
diff --git a/Engine/Source/Developer/CrashDebugHelper/Public/CrashDebugPDBCache.h b/Engine/Source/Developer/CrashDebugHelper/Public/CrashDebugPDBCache.h
--- a/Engine/Source/Developer/CrashDebugHelper/Public/CrashDebugPDBCache.h
+++ b/Engine/Source/Developer/CrashDebugHelper/Public/CrashDebugPDBCache.h
@@ -67,7 +67,7 @@ protected:
 	enum
 	{
 		/** Size of the PDB cache, in GBs. */
-		PDB_CACHE_SIZE_GB = 128,
+		PDB_CACHE_SIZE_GB = 300,
 		MIN_FREESPACE_GB = 64,
 		/** Age of file when it should be deleted from the PDB cache. */
 		DAYS_TO_DELETE_UNUSED_FILES = 14,
