Add command line args to CrashDebugHelper to allow overriding config for PDB cache settings
[UE-31129] - Crash Report server need to alert on Slack when the PDB cache is full

diff --git a/Engine/Source/Developer/CrashDebugHelper/Private/CrashDebugPDBCache.cpp b/Engine/Source/Developer/CrashDebugHelper/Private/CrashDebugPDBCache.cpp
--- a/Engine/Source/Developer/CrashDebugHelper/Private/CrashDebugPDBCache.cpp
+++ b/Engine/Source/Developer/CrashDebugHelper/Private/CrashDebugPDBCache.cpp
@@ -27,7 +27,11 @@ void FPDBCache::Init()
 
 	if (bUsePDBCache)
 	{
-		GConfig->GetString( TEXT( "Engine.CrashDebugHelper" ), TEXT( "DepotRoot" ), DepotRoot, GEngineIni );
+		if (!FParse::Value(FCommandLine::Get(), TEXT("PDBCacheDepotRoot="), DepotRoot))
+		{
+			GConfig->GetString(TEXT("Engine.CrashDebugHelper"), TEXT("DepotRoot"), DepotRoot, GEngineIni);
+		}
+
 		ICrashDebugHelper::SetDepotIndex( DepotRoot );
 
 		const bool bHasDepotRoot = IFileManager::Get().DirectoryExists( *DepotRoot );
@@ -40,10 +44,13 @@ void FPDBCache::Init()
 	// Get the rest of the PDB cache configuration.
 	if( bUsePDBCache )
 	{
-		if( !GConfig->GetString( TEXT( "Engine.CrashDebugHelper" ), TEXT( "PDBCachePath" ), PDBCachePath, GEngineIni ) )
+		if (!FParse::Value(FCommandLine::Get(), TEXT("PDBCachePath="), PDBCachePath))
 		{
-			UE_LOG( LogCrashDebugHelper, Warning, TEXT( "Failed to get PDBCachePath from ini file" ) );
-			bUsePDBCache = false;
+			if (!GConfig->GetString(TEXT("Engine.CrashDebugHelper"), TEXT("PDBCachePath"), PDBCachePath, GEngineIni))
+			{
+				UE_LOG(LogCrashDebugHelper, Warning, TEXT("Failed to get PDBCachePath from ini file or command line"));
+				bUsePDBCache = false;
+			}
 		}
 
 		ICrashDebugHelper::SetDepotIndex( PDBCachePath );
@@ -51,19 +58,28 @@ void FPDBCache::Init()
 
 	if( bUsePDBCache )
 	{
-		if( !GConfig->GetInt( TEXT( "Engine.CrashDebugHelper" ), TEXT( "PDBCacheSizeGB" ), PDBCacheSizeGB, GEngineIni ) )
+		if (!FParse::Value(FCommandLine::Get(), TEXT("PDBCacheSizeGB="), PDBCacheSizeGB))
 		{
-			UE_LOG( LogCrashDebugHelper, Warning, TEXT( "Failed to get PDBCachePath from ini file" ) );
+			if (!GConfig->GetInt(TEXT("Engine.CrashDebugHelper"), TEXT("PDBCacheSizeGB"), PDBCacheSizeGB, GEngineIni))
+			{
+				UE_LOG(LogCrashDebugHelper, Warning, TEXT("Failed to get PDBCachePath from ini file or command line"));
+			}
 		}
 
-		if( !GConfig->GetInt( TEXT( "Engine.CrashDebugHelper" ), TEXT( "MinDiskFreeSpaceGB" ), MinDiskFreeSpaceGB, GEngineIni ) )
+		if (!FParse::Value(FCommandLine::Get(), TEXT("PDBCacheMinFreeSpaceGB="), MinDiskFreeSpaceGB))
 		{
-			UE_LOG( LogCrashDebugHelper, Warning, TEXT( "Failed to get MinDiskFreeSpaceGB from ini file" ) );
+			if (!GConfig->GetInt(TEXT("Engine.CrashDebugHelper"), TEXT("MinDiskFreeSpaceGB"), MinDiskFreeSpaceGB, GEngineIni))
+			{
+				UE_LOG(LogCrashDebugHelper, Warning, TEXT("Failed to get MinDiskFreeSpaceGB from ini file or command line"));
+			}
 		}
 
-		if( !GConfig->GetInt( TEXT( "Engine.CrashDebugHelper" ), TEXT( "DaysToDeleteUnusedFilesFromPDBCache" ), DaysToDeleteUnusedFilesFromPDBCache, GEngineIni ) )
+		if (!FParse::Value(FCommandLine::Get(), TEXT("PDBCacheFileDeleteDays="), DaysToDeleteUnusedFilesFromPDBCache))
 		{
-			UE_LOG( LogCrashDebugHelper, Warning, TEXT( "Failed to get DaysToDeleteUnusedFilesFromPDBCache from ini file" ) );
+			if (!GConfig->GetInt(TEXT("Engine.CrashDebugHelper"), TEXT("DaysToDeleteUnusedFilesFromPDBCache"), DaysToDeleteUnusedFilesFromPDBCache, GEngineIni))
+			{
+				UE_LOG(LogCrashDebugHelper, Warning, TEXT("Failed to get DaysToDeleteUnusedFilesFromPDBCache from ini file or command line"));
+			}
 		}
 
 		InitializePDBCache();
