Add BuildVersion command line arg to MinidumpDiagnostics. Remove PDB cache settings from engine config that are now set via MinidumpDiagnostics command line args.
[UE-31129] - Crash Report server need to alert on Slack when the PDB cache is full

#tests MinidumpDiagnostics tested locally and running okay on CR server

diff --git a/Engine/Source/Developer/CrashDebugHelper/Private/CrashDebugHelper.cpp b/Engine/Source/Developer/CrashDebugHelper/Private/CrashDebugHelper.cpp
--- a/Engine/Source/Developer/CrashDebugHelper/Private/CrashDebugHelper.cpp
+++ b/Engine/Source/Developer/CrashDebugHelper/Private/CrashDebugHelper.cpp
@@ -166,6 +166,18 @@ bool ICrashDebugHelper::Init()
 		CrashInfo.EngineVersion = EngineVersion.ToString();
 	}
 
+	// Check if we have a valid BuildVersion, if so use it.
+	FString CmdBuildVersion;
+	const bool bHasBuildVersion = FParse::Value(FCommandLine::Get(), TEXT("BuildVersion="), CmdBuildVersion);
+	if (bHasBuildVersion)
+	{
+		CrashInfo.BuildVersion = CmdBuildVersion;
+	}
+	else
+	{
+		CrashInfo.BuildVersion = FApp::GetBuildVersion();
+	}
+
 	FString PlatformName;
 	const bool bHasPlatformName = FParse::Value(FCommandLine::Get(), TEXT("PlatformName="), PlatformName);
 	if (bHasPlatformName)
@@ -193,7 +205,8 @@ bool ICrashDebugHelper::Init()
 	UE_LOG( LogCrashDebugHelper, Log, TEXT( "DepotName: %s" ), *CrashInfo.DepotName );
 	UE_LOG( LogCrashDebugHelper, Log, TEXT( "BuiltFromCL: %i" ), CrashInfo.BuiltFromCL );
 	UE_LOG( LogCrashDebugHelper, Log, TEXT( "EngineVersion: %s" ), *CrashInfo.EngineVersion );
-	
+	UE_LOG( LogCrashDebugHelper, Log, TEXT( "BuildVersion: %s" ), *CrashInfo.BuildVersion );
+
 	GConfig->GetString( TEXT( "Engine.CrashDebugHelper" ), TEXT( "SourceControlBuildLabelPattern" ), SourceControlBuildLabelPattern, GEngineIni );
 
 	FCrashDebugHelperConfig::Get().ReadFullCrashDumpConfigurations();
@@ -809,18 +822,21 @@ void ICrashDebugHelper::FindSymbolsAndBinariesStorage()
 	const FString StrPLATFORM_VARIANT = CrashInfo.PlatformVariantName;
 	const FString StrOLD_ENGINE_VERSION = FString::Printf( TEXT( "%s-CL-%i" ), *CrashInfo.DepotName.Replace( TEXT( "+" ), TEXT( "/" ) ), CrashInfo.BuiltFromCL )
 		.Replace( TEXT("/"), TEXT("+") );
+	const FString StrBUILD_VERSION = CrashInfo.BuildVersion;
 
 	const FString TestExecutablesPath = ExecutablePathPattern
 		.Replace( TEXT( "%ENGINE_VERSION%" ), *StrENGINE_VERSION )
 		.Replace( TEXT( "%PLATFORM_NAME%" ), *StrPLATFORM_NAME )
 		.Replace( TEXT( "%PLATFORM_VARIANT%" ), *StrPLATFORM_VARIANT )
-		.Replace( TEXT( "%OLD_ENGINE_VERSION%" ), *StrOLD_ENGINE_VERSION );
+		.Replace( TEXT( "%OLD_ENGINE_VERSION%" ), *StrOLD_ENGINE_VERSION )
+		.Replace( TEXT( "%BUILD_VERSION%" ), *StrBUILD_VERSION );
 
 	const FString TestSymbolsPath = SymbolPathPattern
 		.Replace( TEXT( "%ENGINE_VERSION%" ), *StrENGINE_VERSION )
 		.Replace( TEXT( "%PLATFORM_NAME%" ), *StrPLATFORM_NAME )
 		.Replace( TEXT( "%PLATFORM_VARIANT%" ), *StrPLATFORM_VARIANT )
-		.Replace( TEXT( "%OLD_ENGINE_VERSION%" ), *StrOLD_ENGINE_VERSION );
+		.Replace( TEXT( "%OLD_ENGINE_VERSION%" ), *StrOLD_ENGINE_VERSION )
+		.Replace( TEXT( "%BUILD_VERSION%" ), *StrBUILD_VERSION );
 
 	// Try to find the network path by using the pattern supplied via ini.
 	// If this step successes, we will grab the executable from the network path instead of P4.
diff --git a/Engine/Source/Developer/CrashDebugHelper/Public/CrashDebugHelper.h b/Engine/Source/Developer/CrashDebugHelper/Public/CrashDebugHelper.h
--- a/Engine/Source/Developer/CrashDebugHelper/Public/CrashDebugHelper.h
+++ b/Engine/Source/Developer/CrashDebugHelper/Public/CrashDebugHelper.h
@@ -140,6 +140,9 @@ public:
 	/** Product version, based on FEngineVersion. */
 	FString EngineVersion;
 
+	/** Build version string. */
+	FString BuildVersion;
+
 	/** CL built from. */
 	int32 BuiltFromCL;
 
