Fix MDD callstack stripping bug and add -NoTrimCallstack command line arg to MDD

diff --git a/Engine/Source/Developer/CrashDebugHelper/Private/Windows/CrashDebugHelperWindows.cpp b/Engine/Source/Developer/CrashDebugHelper/Private/Windows/CrashDebugHelperWindows.cpp
--- a/Engine/Source/Developer/CrashDebugHelper/Private/Windows/CrashDebugHelperWindows.cpp
+++ b/Engine/Source/Developer/CrashDebugHelper/Private/Windows/CrashDebugHelperWindows.cpp
@@ -14,6 +14,7 @@ bool FCrashDebugHelperWindows::CreateMinidumpDiagnosticReport( const FString& In
 {
 	const bool bSyncSymbols = FParse::Param( FCommandLine::Get(), TEXT( "SyncSymbols" ) );
 	const bool bAnnotate = FParse::Param( FCommandLine::Get(), TEXT( "Annotate" ) );
+	const bool bNoTrim = FParse::Param(FCommandLine::Get(), TEXT("NoTrimCallstack"));
 	const bool bUseSCC = bSyncSymbols || bAnnotate;
 
 	if( bUseSCC )
@@ -66,7 +67,7 @@ bool FCrashDebugHelperWindows::CreateMinidumpDiagnosticReport( const FString& In
 			WindowsStackWalkExt.GetExceptionInfo();
 
 			// Get the callstacks for each thread
-			bHasAtLeastThreeValidFunctions = WindowsStackWalkExt.GetCallstacks() >= 3;
+			bHasAtLeastThreeValidFunctions = WindowsStackWalkExt.GetCallstacks(!bNoTrim) >= 3;
 
 			// Sync the source file where the crash occurred
 			if( CrashInfo.SourceFile.Len() > 0 )
diff --git a/Engine/Source/Developer/CrashDebugHelper/Private/Windows/WindowsPlatformStackWalkExt.cpp b/Engine/Source/Developer/CrashDebugHelper/Private/Windows/WindowsPlatformStackWalkExt.cpp
--- a/Engine/Source/Developer/CrashDebugHelper/Private/Windows/WindowsPlatformStackWalkExt.cpp
+++ b/Engine/Source/Developer/CrashDebugHelper/Private/Windows/WindowsPlatformStackWalkExt.cpp
@@ -416,7 +416,7 @@ void FWindowsPlatformStackWalkExt::GetExceptionInfo()
 }
 
 
-int FWindowsPlatformStackWalkExt::GetCallstacks()
+int FWindowsPlatformStackWalkExt::GetCallstacks(bool bTrimCallstack)
 {
 	const int32 MAX_NAME_LENGHT = FProgramCounterSymbolInfo::MAX_NAME_LENGHT;
 	int32 NumValidFunctionNames = 0;
@@ -518,7 +518,7 @@ int FWindowsPlatformStackWalkExt::GetCallstacks()
 					|| FunctionName.Contains( TEXT( "NewReportEnsure" ), ESearchCase::CaseSensitive ) )
 				{
 					bFoundSourceFile = false;
-					AssertOrEnsureIndex = FMath::Max( AssertOrEnsureIndex, (int32)StackIndex );
+					AssertOrEnsureIndex = Exception.CallStackString.Num();
 				}
 			}
 
@@ -533,7 +533,7 @@ int FWindowsPlatformStackWalkExt::GetCallstacks()
 	}
 
 	// Remove callstack entries below FDebug, we don't need them.
-	if (AssertOrEnsureIndex > 0)
+	if (bTrimCallstack && AssertOrEnsureIndex > 0)
 	{	
 		Exception.CallStackString.RemoveAt( 0, AssertOrEnsureIndex );
 		UE_LOG( LogCrashDebugHelper, Warning, TEXT( "Callstack trimmed to %i entries" ), Exception.CallStackString.Num() );
diff --git a/Engine/Source/Developer/CrashDebugHelper/Private/Windows/WindowsPlatformStackWalkExt.h b/Engine/Source/Developer/CrashDebugHelper/Private/Windows/WindowsPlatformStackWalkExt.h
--- a/Engine/Source/Developer/CrashDebugHelper/Private/Windows/WindowsPlatformStackWalkExt.h
+++ b/Engine/Source/Developer/CrashDebugHelper/Private/Windows/WindowsPlatformStackWalkExt.h
@@ -52,9 +52,10 @@ struct FWindowsPlatformStackWalkExt
 	/**
 	 * Gets the callstack of the crash.
 	 *
+	 * @param bTrimCallstack	If true, trims what it thinks are irrelevant entries after a debug or assert. If false, leaves all callstack entries. 
 	 * @return the number of valid function names
 	 */
-	int32 GetCallstacks();
+	int32 GetCallstacks(bool bTrimCallstack);
 
 	/** Opens a minidump as a new session. */
 	bool OpenDumpFile( const FString& DumpFileName );
