Added command line args for configuring CrashDebugHelper to lock access to P4 and PDB cache
[UE-34402] - Crash Reporter: Improve CRP performance by allowing multiple MDD instances

diff --git a/Engine/Source/Developer/CrashDebugHelper/Private/CrashDebugHelper.cpp b/Engine/Source/Developer/CrashDebugHelper/Private/CrashDebugHelper.cpp
--- a/Engine/Source/Developer/CrashDebugHelper/Private/CrashDebugHelper.cpp
+++ b/Engine/Source/Developer/CrashDebugHelper/Private/CrashDebugHelper.cpp
@@ -143,6 +143,9 @@ bool ICrashDebugHelper::Init()
 {
 	bInitialized = true;
 
+	CrashInfo.bMutexPDBCache = FParse::Param(FCommandLine::Get(), TEXT("MutexPDBCache"));
+	FParse::Value(FCommandLine::Get(), TEXT("PDBCacheLock="), CrashInfo.PDBCacheLockName);
+
 	// Check if we have a valid EngineVersion, if so use it.
 	FString CmdEngineVersion;
 	const bool bHasEngineVersion = FParse::Value( FCommandLine::Get(), TEXT( "EngineVersion=" ), CmdEngineVersion );
@@ -216,13 +219,29 @@ bool ICrashDebugHelper::Init()
 
 	if (bUsePDBCache)
 	{
-		FPDBCache::Get().Init();
+		if (CrashInfo.bMutexPDBCache && !CrashInfo.PDBCacheLockName.IsEmpty())
+		{
+			// Scoped lock
+			FSystemWideCriticalSection PDBCacheLock(CrashInfo.PDBCacheLockName, FTimespan(0, 0, 3, 0, 0));
+			if (PDBCacheLock.IsValid())
+			{
+				FPDBCache::Get().Init();
+			}
+		}
+		else
+		{
+			FPDBCache::Get().Init();
+		}
+
+		if (!FPDBCache::Get().UsePDBCache())
+		{
+			UE_LOG(LogCrashDebugHelper, Warning, TEXT("PDB Cache failed to initialize"));
+		}
 	}
 	else
 	{
 		UE_LOG( LogCrashDebugHelper, Warning, TEXT( "PDB Cache disabled" ) );
-	}
-	
+	}	
 
 	return bInitialized;
 }
@@ -277,12 +296,6 @@ void ICrashDebugHelper::ShutdownSourceControl()
 
 bool ICrashDebugHelper::SyncModules()
 {
-	// Check source control
-	if( !ISourceControlModule::Get().IsEnabled() )
-	{
-		return false;
-	}
-
 	if( !FPDBCache::Get().UsePDBCache() )
 	{
 		UE_LOG( LogCrashDebugHelper, Warning, TEXT( "The PDB Cache is disabled, cannot proceed, %s" ), *CrashInfo.EngineVersion );
@@ -293,7 +306,6 @@ bool ICrashDebugHelper::SyncModules()
 	const TCHAR* UESymbols = TEXT( "Rocket/Symbols/" );
 	const bool bHasExecutable = !CrashInfo.ExecutablesPath.IsEmpty();
 	const bool bHasSymbols = !CrashInfo.SymbolsPath.IsEmpty();
-	TArray< TSharedRef<ISourceControlLabel> > Labels = ISourceControlModule::Get().GetProvider().GetLabels( CrashInfo.LabelName );
 	
 	const bool bContainsProductVersion = FPDBCache::Get().ContainsPDBCacheEntry( CrashInfo.EngineVersion );
 	if( bHasExecutable && bHasSymbols )
@@ -343,173 +355,186 @@ bool ICrashDebugHelper::SyncModules()
 			CrashInfo.PDBCacheEntry = FPDBCache::Get().CreateAndAddPDBCacheEntryMixed( CrashInfo.EngineVersion, FilesToBeCached );
 		}
 	}
-	// Get all labels associated with the crash info's label.
 	// OBSOLETE PATH
-	else if( Labels.Num() >= 1 )
+	else
 	{
-		TSharedRef<ISourceControlLabel> Label = Labels[0];
-		TSet<FString> FilesToSync;
-
-		// Use product version instead of label name to make a distinguish between chosen methods.
-		const bool bContainsLabelName = FPDBCache::Get().ContainsPDBCacheEntry( CrashInfo.LabelName );
+		// Command line for blocking obsolete path
+		const bool bNoP4Symbols = FParse::Param(FCommandLine::Get(), TEXT("NoP4Symbols"));
 
-		if( bContainsProductVersion )
+		// Check source control
+		if (bNoP4Symbols || !ISourceControlModule::Get().IsEnabled())
 		{
-			UE_LOG( LogCrashDebugHelper, Warning, TEXT( "Using cached storage: %s" ), *CrashInfo.EngineVersion );
-			CrashInfo.PDBCacheEntry = FPDBCache::Get().FindAndTouchPDBCacheEntry( CrashInfo.EngineVersion );
+			return false;
 		}
-		else if( bContainsLabelName )
+
+		// Get all labels associated with the crash info's label.
+		TArray< TSharedRef<ISourceControlLabel> > Labels = ISourceControlModule::Get().GetProvider().GetLabels(CrashInfo.LabelName);
+
+		if (Labels.Num() >= 1)
 		{
-			UE_LOG( LogCrashDebugHelper, Warning, TEXT( "Using cached storage: %s" ), *CrashInfo.LabelName );
-			CrashInfo.PDBCacheEntry = FPDBCache::Get().FindAndTouchPDBCacheEntry( CrashInfo.LabelName );
-		}
-		else if( bHasExecutable )
-		{			
-			SCOPE_LOG_TIME_IN_SECONDS( TEXT( "SyncModulesAndNetwork" ), nullptr );
+			TSharedRef<ISourceControlLabel> Label = Labels[0];
+			TSet<FString> FilesToSync;
 
-			// Grab information about symbols.
-			TArray< TSharedRef<class ISourceControlRevision, ESPMode::ThreadSafe> > PDBSourceControlRevisions;
-			const FString PDBsPath = FString::Printf( TEXT( "%s/%s....pdb" ), *CrashInfo.DepotName, UESymbols );
-			Label->GetFileRevisions( PDBsPath, PDBSourceControlRevisions );
+			// Use product version instead of label name to make a distinguish between chosen methods.
+			const bool bContainsLabelName = FPDBCache::Get().ContainsPDBCacheEntry(CrashInfo.LabelName);
 
-			TSet<FString> PDBPaths;
-			for( const auto& PDBSrc : PDBSourceControlRevisions )
+			if (bContainsProductVersion)
 			{
-				PDBPaths.Add( PDBSrc->GetFilename() );
+				UE_LOG(LogCrashDebugHelper, Warning, TEXT("Using cached storage: %s"), *CrashInfo.EngineVersion);
+				CrashInfo.PDBCacheEntry = FPDBCache::Get().FindAndTouchPDBCacheEntry(CrashInfo.EngineVersion);
 			}
-
-			// Now, sync symbols.
-			for( const auto& PDBPath : PDBPaths )
+			else if (bContainsLabelName)
 			{
-				if( Label->Sync( PDBPath ) )
-				{
-					UE_LOG( LogCrashDebugHelper, Warning, TEXT( "Synced PDB: %s" ), *PDBPath );
-				}
+				UE_LOG(LogCrashDebugHelper, Warning, TEXT("Using cached storage: %s"), *CrashInfo.LabelName);
+				CrashInfo.PDBCacheEntry = FPDBCache::Get().FindAndTouchPDBCacheEntry(CrashInfo.LabelName);
 			}
+			else if (bHasExecutable)
+			{
+				SCOPE_LOG_TIME_IN_SECONDS(TEXT("SyncModulesAndNetwork"), nullptr);
 
-			// Find all the executables in the product network path.
-			TArray<FString> NetworkExecutables;
-			IFileManager::Get().FindFilesRecursive( NetworkExecutables, *CrashInfo.ExecutablesPath, TEXT( "*.dll" ), true, false, false );
-			IFileManager::Get().FindFilesRecursive( NetworkExecutables, *CrashInfo.ExecutablesPath, TEXT( "*.exe" ), true, false, false );
+				// Grab information about symbols.
+				TArray< TSharedRef<class ISourceControlRevision, ESPMode::ThreadSafe> > PDBSourceControlRevisions;
+				const FString PDBsPath = FString::Printf(TEXT("%s/%s....pdb"), *CrashInfo.DepotName, UESymbols);
+				Label->GetFileRevisions(PDBsPath, PDBSourceControlRevisions);
 
-			// From=Full pathname
-			// To=Relative pathname
-			TMap<FString, FString> FilesToBeCached;
+				TSet<FString> PDBPaths;
+				for (const auto& PDBSrc : PDBSourceControlRevisions)
+				{
+					PDBPaths.Add(PDBSrc->GetFilename());
+				}
 
-			// If a symbol matches an executable, add the pair to the list of files that should be cached.
-			for( const auto& NetworkExecutableFullpath : NetworkExecutables )
-			{
-				for( const auto& PDBPath : PDBPaths )
+				// Now, sync symbols.
+				for (const auto& PDBPath : PDBPaths)
 				{
-					const FString PDBRelativePath = PDBPath.Replace( *CrashInfo.DepotName, TEXT( "" ) ).Replace( UESymbols, TEXT( "" ) );
-					const FString PDBFullpath = FPDBCache::Get().GetDepotRoot() / PDBPath;
+					if (Label->Sync(PDBPath))
+					{
+						UE_LOG(LogCrashDebugHelper, Warning, TEXT("Synced PDB: %s"), *PDBPath);
+					}
+				}
+
+				// Find all the executables in the product network path.
+				TArray<FString> NetworkExecutables;
+				IFileManager::Get().FindFilesRecursive(NetworkExecutables, *CrashInfo.ExecutablesPath, TEXT("*.dll"), true, false, false);
+				IFileManager::Get().FindFilesRecursive(NetworkExecutables, *CrashInfo.ExecutablesPath, TEXT("*.exe"), true, false, false);
+
+				// From=Full pathname
+				// To=Relative pathname
+				TMap<FString, FString> FilesToBeCached;
 
-					const FString PDBMatch = PDBRelativePath.Replace( TEXT( "pdb" ), TEXT( "" ) );
-					const FString NetworkRelativePath = NetworkExecutableFullpath.Replace( *CrashInfo.ExecutablesPath, TEXT( "" ) );
-					const bool bMatch = NetworkExecutableFullpath.Contains( PDBMatch );
-					if( bMatch )
+				// If a symbol matches an executable, add the pair to the list of files that should be cached.
+				for (const auto& NetworkExecutableFullpath : NetworkExecutables)
+				{
+					for (const auto& PDBPath : PDBPaths)
 					{
-						// From -> Where
-						FilesToBeCached.Add( NetworkExecutableFullpath, NetworkRelativePath );
-						FilesToBeCached.Add( PDBFullpath, PDBRelativePath );
-						break;
+						const FString PDBRelativePath = PDBPath.Replace(*CrashInfo.DepotName, TEXT("")).Replace(UESymbols, TEXT(""));
+						const FString PDBFullpath = FPDBCache::Get().GetDepotRoot() / PDBPath;
+
+						const FString PDBMatch = PDBRelativePath.Replace(TEXT("pdb"), TEXT(""));
+						const FString NetworkRelativePath = NetworkExecutableFullpath.Replace(*CrashInfo.ExecutablesPath, TEXT(""));
+						const bool bMatch = NetworkExecutableFullpath.Contains(PDBMatch);
+						if (bMatch)
+						{
+							// From -> Where
+							FilesToBeCached.Add(NetworkExecutableFullpath, NetworkRelativePath);
+							FilesToBeCached.Add(PDBFullpath, PDBRelativePath);
+							break;
+						}
 					}
 				}
-			}
 
-			// Initialize and add a new PDB Cache entry to the database.
-			CrashInfo.PDBCacheEntry = FPDBCache::Get().CreateAndAddPDBCacheEntryMixed( CrashInfo.EngineVersion, FilesToBeCached );
-		}
-		else
-		{
-			TArray<FString> FilesToBeCached;
-			
-			//@TODO: MAC: Excluding labels for Mac since we are only syncing windows binaries here...
-			if( Label->GetName().Contains( TEXT( "Mac" ) ) )
-			{
-				UE_LOG( LogCrashDebugHelper, Log, TEXT( "Skipping Mac label: %s" ), *Label->GetName() );
+				// Initialize and add a new PDB Cache entry to the database.
+				CrashInfo.PDBCacheEntry = FPDBCache::Get().CreateAndAddPDBCacheEntryMixed(CrashInfo.EngineVersion, FilesToBeCached);
 			}
 			else
 			{
-				// Sync all the dll, exes, and related symbol files
-				UE_LOG( LogCrashDebugHelper, Log, TEXT( "Syncing modules with label: %s" ), *Label->GetName() );
+				TArray<FString> FilesToBeCached;
 
-				SCOPE_LOG_TIME_IN_SECONDS( TEXT( "SyncModules" ), nullptr );
+				//@TODO: MAC: Excluding labels for Mac since we are only syncing windows binaries here...
+				if (Label->GetName().Contains(TEXT("Mac")))
+				{
+					UE_LOG(LogCrashDebugHelper, Log, TEXT("Skipping Mac label: %s"), *Label->GetName());
+				}
+				else
+				{
+					// Sync all the dll, exes, and related symbol files
+					UE_LOG(LogCrashDebugHelper, Log, TEXT("Syncing modules with label: %s"), *Label->GetName());
 
-				// Grab all dll and pdb files for the specified label.
-				TArray< TSharedRef<class ISourceControlRevision, ESPMode::ThreadSafe> > DLLSourceControlRevisions;
-				const FString DLLsPath = FString::Printf( TEXT( "%s/....dll" ), *CrashInfo.DepotName );
-				Label->GetFileRevisions( DLLsPath, DLLSourceControlRevisions );
+					SCOPE_LOG_TIME_IN_SECONDS(TEXT("SyncModules"), nullptr);
 
-				TArray< TSharedRef<class ISourceControlRevision, ESPMode::ThreadSafe> > EXESourceControlRevisions;
-				const FString EXEsPath = FString::Printf( TEXT( "%s/....exe" ), *CrashInfo.DepotName );
-				Label->GetFileRevisions( EXEsPath, EXESourceControlRevisions );
+					// Grab all dll and pdb files for the specified label.
+					TArray< TSharedRef<class ISourceControlRevision, ESPMode::ThreadSafe> > DLLSourceControlRevisions;
+					const FString DLLsPath = FString::Printf(TEXT("%s/....dll"), *CrashInfo.DepotName);
+					Label->GetFileRevisions(DLLsPath, DLLSourceControlRevisions);
 
-				TArray< TSharedRef<class ISourceControlRevision, ESPMode::ThreadSafe> > PDBSourceControlRevisions;
-				const FString PDBsPath = FString::Printf( TEXT( "%s/....pdb" ), *CrashInfo.DepotName );
-				Label->GetFileRevisions( PDBsPath, PDBSourceControlRevisions );
+					TArray< TSharedRef<class ISourceControlRevision, ESPMode::ThreadSafe> > EXESourceControlRevisions;
+					const FString EXEsPath = FString::Printf(TEXT("%s/....exe"), *CrashInfo.DepotName);
+					Label->GetFileRevisions(EXEsPath, EXESourceControlRevisions);
 
-				TSet<FString> ModulesPaths;
-				for( const auto& DLLSrc : DLLSourceControlRevisions )
-				{
-					ModulesPaths.Add( DLLSrc->GetFilename().Replace( *CrashInfo.DepotName, TEXT( "" ) ) );
-				}
-				for( const auto& EXESrc : EXESourceControlRevisions )
-				{
-					ModulesPaths.Add( EXESrc->GetFilename().Replace( *CrashInfo.DepotName, TEXT( "" ) ) );
-				}
+					TArray< TSharedRef<class ISourceControlRevision, ESPMode::ThreadSafe> > PDBSourceControlRevisions;
+					const FString PDBsPath = FString::Printf(TEXT("%s/....pdb"), *CrashInfo.DepotName);
+					Label->GetFileRevisions(PDBsPath, PDBSourceControlRevisions);
 
-				TSet<FString> PDBPaths;
-				for( const auto& PDBSrc : PDBSourceControlRevisions )
-				{
-					PDBPaths.Add( PDBSrc->GetFilename().Replace( *CrashInfo.DepotName, TEXT( "" ) ) );
-				}
+					TSet<FString> ModulesPaths;
+					for (const auto& DLLSrc : DLLSourceControlRevisions)
+					{
+						ModulesPaths.Add(DLLSrc->GetFilename().Replace(*CrashInfo.DepotName, TEXT("")));
+					}
+					for (const auto& EXESrc : EXESourceControlRevisions)
+					{
+						ModulesPaths.Add(EXESrc->GetFilename().Replace(*CrashInfo.DepotName, TEXT("")));
+					}
 
-				// Iterate through all module and see if we have dll and pdb associated with the module, if so add it to the files to sync.
-				for( const auto& ModuleName : CrashInfo.ModuleNames )
-				{
-					const FString ModuleNamePDB = ModuleName.Replace( TEXT( ".dll" ), TEXT( ".pdb" ) ).Replace( TEXT( ".exe" ), TEXT( ".pdb" ) );
+					TSet<FString> PDBPaths;
+					for (const auto& PDBSrc : PDBSourceControlRevisions)
+					{
+						PDBPaths.Add(PDBSrc->GetFilename().Replace(*CrashInfo.DepotName, TEXT("")));
+					}
 
-					for( const auto& ModulePath : ModulesPaths )
+					// Iterate through all module and see if we have dll and pdb associated with the module, if so add it to the files to sync.
+					for (const auto& ModuleName : CrashInfo.ModuleNames)
 					{
-						const bool bContainsModule = ModulePath.Contains( ModuleName );
-						if( bContainsModule )
+						const FString ModuleNamePDB = ModuleName.Replace(TEXT(".dll"), TEXT(".pdb")).Replace(TEXT(".exe"), TEXT(".pdb"));
+
+						for (const auto& ModulePath : ModulesPaths)
 						{
-							FilesToSync.Add( ModulePath );
+							const bool bContainsModule = ModulePath.Contains(ModuleName);
+							if (bContainsModule)
+							{
+								FilesToSync.Add(ModulePath);
+							}
 						}
-					}
 
-					for( const auto& PDBPath : PDBPaths )
-					{
-						const bool bContainsPDB = PDBPath.Contains( ModuleNamePDB );
-						if( bContainsPDB )
+						for (const auto& PDBPath : PDBPaths)
 						{
-							FilesToSync.Add( PDBPath );
+							const bool bContainsPDB = PDBPath.Contains(ModuleNamePDB);
+							if (bContainsPDB)
+							{
+								FilesToSync.Add(PDBPath);
+							}
 						}
 					}
-				}
 
-				// Now, sync all files.
-				for( const auto& Filename : FilesToSync )
-				{
-					const FString DepotPath = CrashInfo.DepotName + Filename;
-					if( Label->Sync( DepotPath ) )
+					// Now, sync all files.
+					for (const auto& Filename : FilesToSync)
 					{
-						UE_LOG( LogCrashDebugHelper, Warning, TEXT( "Synced binary: %s" ), *DepotPath );
+						const FString DepotPath = CrashInfo.DepotName + Filename;
+						if (Label->Sync(DepotPath))
+						{
+							UE_LOG(LogCrashDebugHelper, Warning, TEXT("Synced binary: %s"), *DepotPath);
+						}
+						FilesToBeCached.Add(DepotPath);
 					}
-					FilesToBeCached.Add( DepotPath );
 				}
-			}
 
-			// Initialize and add a new PDB Cache entry to the database.
-			CrashInfo.PDBCacheEntry = FPDBCache::Get().CreateAndAddPDBCacheEntry( CrashInfo.LabelName, CrashInfo.DepotName, FilesToBeCached );
+				// Initialize and add a new PDB Cache entry to the database.
+				CrashInfo.PDBCacheEntry = FPDBCache::Get().CreateAndAddPDBCacheEntry(CrashInfo.LabelName, CrashInfo.DepotName, FilesToBeCached);
+			}
+		}
+		else
+		{
+			UE_LOG(LogCrashDebugHelper, Error, TEXT("Could not find label: %s"), *CrashInfo.LabelName);
+			return false;
 		}
-	}
-	// OBSOLETE PATH
-	else
-	{
-		UE_LOG( LogCrashDebugHelper, Error, TEXT( "Could not find label: %s"), *CrashInfo.LabelName );
-		return false;
 	}
 
 	return true;
diff --git a/Engine/Source/Developer/CrashDebugHelper/Private/Windows/CrashDebugHelperWindows.cpp b/Engine/Source/Developer/CrashDebugHelper/Private/Windows/CrashDebugHelperWindows.cpp
--- a/Engine/Source/Developer/CrashDebugHelper/Private/Windows/CrashDebugHelperWindows.cpp
+++ b/Engine/Source/Developer/CrashDebugHelper/Private/Windows/CrashDebugHelperWindows.cpp
@@ -29,27 +29,26 @@ bool FCrashDebugHelperWindows::CreateMinidumpDiagnosticReport( const FString& In
 	bool bHasAtLeastThreeValidFunctions = false;
 	if( bReady && WindowsStackWalkExt.OpenDumpFile( InCrashDumpFilename ) )
 	{
-		if( CrashInfo.BuiltFromCL != FCrashInfo::INVALID_CHANGELIST )
+		if (CrashInfo.BuiltFromCL != FCrashInfo::INVALID_CHANGELIST)
 		{
 			// Get the build version and modules paths.
 			FCrashModuleInfo ExeFileVersion;
-			WindowsStackWalkExt.GetExeFileVersionAndModuleList( ExeFileVersion );
+			WindowsStackWalkExt.GetExeFileVersionAndModuleList(ExeFileVersion);
 
-			// CrashInfo now contains a changelist to lookup a label for
-			if( bSyncSymbols )
+			// Init Symbols
+			if (CrashInfo.bMutexPDBCache && !CrashInfo.PDBCacheLockName.IsEmpty())
 			{
-				FindSymbolsAndBinariesStorage();
-
-				const bool bSynced = SyncModules();
-				// Without symbols we can't decode the provided minidump.
-				if( !bSynced )
+				// Scoped lock
+				FSystemWideCriticalSection PDBCacheLock(CrashInfo.PDBCacheLockName, FTimespan(0, 0, 3, 0, 0));
+				if (PDBCacheLock.IsValid())
 				{
-					return 0;
+					InitSymbols(WindowsStackWalkExt, bSyncSymbols);
 				}
 			}
-
-			// Initialise the symbol options
-			WindowsStackWalkExt.InitSymbols();
+			else
+			{
+				InitSymbols(WindowsStackWalkExt, bSyncSymbols);
+			}
 
 			// Set the symbol path based on the loaded modules
 			WindowsStackWalkExt.SetSymbolPathsFromModules();
@@ -72,23 +71,23 @@ bool FCrashDebugHelperWindows::CreateMinidumpDiagnosticReport( const FString& In
 			// Sync the source file where the crash occurred
 			if( CrashInfo.SourceFile.Len() > 0 )
 			{
-				if( bSyncSymbols && CrashInfo.BuiltFromCL > 0 )
-				{
-					UE_LOG( LogCrashDebugHelper, Log, TEXT( "Using CL %i to sync crash source file" ), CrashInfo.BuiltFromCL );
-					SyncSourceFile();
-				}
+				const bool bMutexSourceSync = FParse::Param(FCommandLine::Get(), TEXT("MutexSourceSync"));
+				FString SourceSyncLockName;
+				FParse::Value(FCommandLine::Get(), TEXT("SourceSyncLock="), SourceSyncLockName);
 
-				// Try to annotate the file if requested
-				bool bAnnotationSuccessful = false;
-				if( bAnnotate )
+				if (bMutexSourceSync && !SourceSyncLockName.IsEmpty())
 				{
-					bAnnotationSuccessful = AddAnnotatedSourceToReport();
+					// Scoped lock
+					const FTimespan GlobalLockWaitTimeout(0, 0, 0, 30, 0);
+					FSystemWideCriticalSection SyncSourceLock(SourceSyncLockName, GlobalLockWaitTimeout);
+					if (SyncSourceLock.IsValid())
+					{
+						SyncAndReadSourceFile(bSyncSymbols, bAnnotate, CrashInfo.BuiltFromCL);
+					}	
 				}
-
-				// If annotation is not requested, or failed, add the standard source context
-				if( !bAnnotationSuccessful )
+				else
 				{
-					AddSourceToReport();
+					SyncAndReadSourceFile(bSyncSymbols, bAnnotate, CrashInfo.BuiltFromCL);
 				}
 			}
 		}
@@ -110,4 +109,50 @@ bool FCrashDebugHelperWindows::CreateMinidumpDiagnosticReport( const FString& In
 	return bHasAtLeastThreeValidFunctions;
 }
 
+bool FCrashDebugHelperWindows::InitSymbols(FWindowsPlatformStackWalkExt& WindowsStackWalkExt, bool bSyncSymbols)
+{
+	// CrashInfo now contains a changelist to lookup a label for
+	if (bSyncSymbols)
+	{
+		FindSymbolsAndBinariesStorage();
+
+		const bool bSynced = SyncModules();
+		// Without symbols we can't decode the provided minidump.
+		if (!bSynced)
+		{
+			return false;
+		}
+	}
+
+	// Initialise the symbol options
+	WindowsStackWalkExt.InitSymbols();
+
+	// Set the symbol path based on the loaded modules
+	WindowsStackWalkExt.SetSymbolPathsFromModules();
+
+	return true;
+}
+
+void FCrashDebugHelperWindows::SyncAndReadSourceFile(bool bSyncSymbols, bool bAnnotate, int32 BuiltFromCL)
+{
+	if (bSyncSymbols && BuiltFromCL > 0)
+	{
+		UE_LOG(LogCrashDebugHelper, Log, TEXT("Using CL %i to sync crash source file"), BuiltFromCL);
+		SyncSourceFile();
+	}
+
+	// Try to annotate the file if requested
+	bool bAnnotationSuccessful = false;
+	if (bAnnotate)
+	{
+		bAnnotationSuccessful = AddAnnotatedSourceToReport();
+	}
+
+	// If annotation is not requested, or failed, add the standard source context
+	if (!bAnnotationSuccessful)
+	{
+		AddSourceToReport();
+	}
+}
+
 #include "HideWindowsPlatformTypes.h"
diff --git a/Engine/Source/Developer/CrashDebugHelper/Private/Windows/CrashDebugHelperWindows.h b/Engine/Source/Developer/CrashDebugHelper/Private/Windows/CrashDebugHelperWindows.h
--- a/Engine/Source/Developer/CrashDebugHelper/Private/Windows/CrashDebugHelperWindows.h
+++ b/Engine/Source/Developer/CrashDebugHelper/Private/Windows/CrashDebugHelperWindows.h
@@ -15,6 +15,10 @@ public:
 	 *	@return	bool				true if successful, false if not
 	 */
 	virtual bool CreateMinidumpDiagnosticReport( const FString& InCrashDumpName ) override;
+
+private:
+	bool InitSymbols(struct FWindowsPlatformStackWalkExt& WindowsStackWalkExt, bool bSyncSymbols);
+	void SyncAndReadSourceFile(bool bSyncSymbols, bool bAnnotate, int32 BuiltFromCL);
 };
 
 typedef FCrashDebugHelperWindows FCrashDebugHelper;
diff --git a/Engine/Source/Developer/CrashDebugHelper/Public/CrashDebugHelper.h b/Engine/Source/Developer/CrashDebugHelper/Public/CrashDebugHelper.h
--- a/Engine/Source/Developer/CrashDebugHelper/Public/CrashDebugHelper.h
+++ b/Engine/Source/Developer/CrashDebugHelper/Public/CrashDebugHelper.h
@@ -173,6 +173,12 @@ public:
 	FString PlatformName;
 	FString PlatformVariantName;
 
+	/** If we are using a PDBCache, this is whether we should use a system-wide lock to access it. */
+	bool bMutexPDBCache;
+
+	/** If we are using a PDBCache, this is the name of the system-wide lock we should use to access it. */
+	FString PDBCacheLockName;
+
 	FCrashInfo()
 		: BuiltFromCL( INVALID_CHANGELIST )
 		, SourceLineNumber( 0 )
