Crash Report server. Fixed MinidumpDiagnostics speed problems by adding option to early abort runs that have no valid PDB cache entry. Server was spending too long on reports that can't be symbolicated.

diff --git a/Engine/Source/Developer/CrashDebugHelper/Private/CrashDebugHelper.cpp b/Engine/Source/Developer/CrashDebugHelper/Private/CrashDebugHelper.cpp
--- a/Engine/Source/Developer/CrashDebugHelper/Private/CrashDebugHelper.cpp
+++ b/Engine/Source/Developer/CrashDebugHelper/Private/CrashDebugHelper.cpp
@@ -294,8 +294,10 @@ void ICrashDebugHelper::ShutdownSourceControl()
 }
 
 
-bool ICrashDebugHelper::SyncModules()
+bool ICrashDebugHelper::SyncModules(bool& bOutPDBCacheEntryValid)
 {
+	bOutPDBCacheEntryValid = false;
+
 	if( !FPDBCache::Get().UsePDBCache() )
 	{
 		UE_LOG( LogCrashDebugHelper, Warning, TEXT( "The PDB Cache is disabled, cannot proceed, %s" ), *CrashInfo.EngineVersion );
@@ -537,6 +539,7 @@ bool ICrashDebugHelper::SyncModules()
 		}
 	}
 
+	bOutPDBCacheEntryValid = CrashInfo.PDBCacheEntry.IsValid() && CrashInfo.PDBCacheEntry->Files.Num() > 0;
 	return true;
 }
 
diff --git a/Engine/Source/Developer/CrashDebugHelper/Private/Mac/CrashDebugHelperMac.cpp b/Engine/Source/Developer/CrashDebugHelper/Private/Mac/CrashDebugHelperMac.cpp
--- a/Engine/Source/Developer/CrashDebugHelper/Private/Mac/CrashDebugHelperMac.cpp
+++ b/Engine/Source/Developer/CrashDebugHelper/Private/Mac/CrashDebugHelperMac.cpp
@@ -662,7 +662,8 @@ bool FCrashDebugHelperMac::CreateMinidumpDiagnosticReport( const FString& InCras
 				{
 					FindSymbolsAndBinariesStorage();
 					
-					SyncModules();
+					bool bPDBCacheEntryValid = false;
+					SyncModules(bPDBCacheEntryValid);
 				}
 			}
 			
diff --git a/Engine/Source/Developer/CrashDebugHelper/Private/Windows/CrashDebugHelperWindows.cpp b/Engine/Source/Developer/CrashDebugHelper/Private/Windows/CrashDebugHelperWindows.cpp
--- a/Engine/Source/Developer/CrashDebugHelper/Private/Windows/CrashDebugHelperWindows.cpp
+++ b/Engine/Source/Developer/CrashDebugHelper/Private/Windows/CrashDebugHelperWindows.cpp
@@ -36,60 +36,72 @@ bool FCrashDebugHelperWindows::CreateMinidumpDiagnosticReport( const FString& In
 			WindowsStackWalkExt.GetExeFileVersionAndModuleList(ExeFileVersion);
 
 			// Init Symbols
+			bool bInitSymbols = false;
 			if (CrashInfo.bMutexPDBCache && !CrashInfo.PDBCacheLockName.IsEmpty())
 			{
 				// Scoped lock
+				UE_LOG(LogCrashDebugHelper, Log, TEXT("Locking for InitSymbols()"));
 				FSystemWideCriticalSection PDBCacheLock(CrashInfo.PDBCacheLockName, FTimespan(0, 0, 3, 0, 0));
 				if (PDBCacheLock.IsValid())
 				{
-					InitSymbols(WindowsStackWalkExt, bSyncSymbols);
+					bInitSymbols = InitSymbols(WindowsStackWalkExt, bSyncSymbols);
 				}
+				UE_LOG(LogCrashDebugHelper, Log, TEXT("Unlocking after InitSymbols()"));
 			}
 			else
 			{
-				InitSymbols(WindowsStackWalkExt, bSyncSymbols);
+				bInitSymbols = InitSymbols(WindowsStackWalkExt, bSyncSymbols);
 			}
 
-			// Set the symbol path based on the loaded modules
-			WindowsStackWalkExt.SetSymbolPathsFromModules();
-
-			// Get all the info we should ever need about the modules
-			WindowsStackWalkExt.GetModuleInfoDetailed();
+			if (bInitSymbols)
+			{
+				// Set the symbol path based on the loaded modules
+				WindowsStackWalkExt.SetSymbolPathsFromModules();
 
-			// Get info about the system that created the minidump
-			WindowsStackWalkExt.GetSystemInfo();
+				// Get all the info we should ever need about the modules
+				WindowsStackWalkExt.GetModuleInfoDetailed();
 
-			// Get all the thread info
-			WindowsStackWalkExt.GetThreadInfo();
+				// Get info about the system that created the minidump
+				WindowsStackWalkExt.GetSystemInfo();
 
-			// Get exception info
-			WindowsStackWalkExt.GetExceptionInfo();
+				// Get all the thread info
+				WindowsStackWalkExt.GetThreadInfo();
 
-			// Get the callstacks for each thread
-			bHasAtLeastThreeValidFunctions = WindowsStackWalkExt.GetCallstacks(!bNoTrim) >= 3;
+				// Get exception info
+				WindowsStackWalkExt.GetExceptionInfo();
 
-			// Sync the source file where the crash occurred
-			if( CrashInfo.SourceFile.Len() > 0 )
-			{
-				const bool bMutexSourceSync = FParse::Param(FCommandLine::Get(), TEXT("MutexSourceSync"));
-				FString SourceSyncLockName;
-				FParse::Value(FCommandLine::Get(), TEXT("SourceSyncLock="), SourceSyncLockName);
+				// Get the callstacks for each thread
+				bHasAtLeastThreeValidFunctions = WindowsStackWalkExt.GetCallstacks(!bNoTrim) >= 3;
 
-				if (bMutexSourceSync && !SourceSyncLockName.IsEmpty())
+				// Sync the source file where the crash occurred
+				if (CrashInfo.SourceFile.Len() > 0)
 				{
-					// Scoped lock
-					const FTimespan GlobalLockWaitTimeout(0, 0, 0, 30, 0);
-					FSystemWideCriticalSection SyncSourceLock(SourceSyncLockName, GlobalLockWaitTimeout);
-					if (SyncSourceLock.IsValid())
+					const bool bMutexSourceSync = FParse::Param(FCommandLine::Get(), TEXT("MutexSourceSync"));
+					FString SourceSyncLockName;
+					FParse::Value(FCommandLine::Get(), TEXT("SourceSyncLock="), SourceSyncLockName);
+
+					if (bMutexSourceSync && !SourceSyncLockName.IsEmpty())
+					{
+						// Scoped lock
+						UE_LOG(LogCrashDebugHelper, Log, TEXT("Locking for SyncAndReadSourceFile()"));
+						const FTimespan GlobalLockWaitTimeout(0, 0, 0, 30, 0);
+						FSystemWideCriticalSection SyncSourceLock(SourceSyncLockName, GlobalLockWaitTimeout);
+						if (SyncSourceLock.IsValid())
+						{
+							SyncAndReadSourceFile(bSyncSymbols, bAnnotate, CrashInfo.BuiltFromCL);
+						}
+						UE_LOG(LogCrashDebugHelper, Log, TEXT("Unlocking after SyncAndReadSourceFile()"));
+					}
+					else
 					{
 						SyncAndReadSourceFile(bSyncSymbols, bAnnotate, CrashInfo.BuiltFromCL);
-					}	
-				}
-				else
-				{
-					SyncAndReadSourceFile(bSyncSymbols, bAnnotate, CrashInfo.BuiltFromCL);
+					}
 				}
 			}
+			else
+			{
+				UE_LOG(LogCrashDebugHelper, Warning, TEXT("InitSymbols failed"));
+			}
 		}
 		else
 		{
@@ -116,12 +128,24 @@ bool FCrashDebugHelperWindows::InitSymbols(FWindowsPlatformStackWalkExt& Windows
 	{
 		FindSymbolsAndBinariesStorage();
 
-		const bool bSynced = SyncModules();
+		bool bPDBCacheEntryValid = false;
+		const bool bSynced = SyncModules(bPDBCacheEntryValid);
 		// Without symbols we can't decode the provided minidump.
 		if (!bSynced)
 		{
 			return false;
 		}
+
+		if (!bPDBCacheEntryValid)
+		{
+			// early-out option
+			const bool bForceUsePDBCache = FParse::Param(FCommandLine::Get(), TEXT("ForceUsePDBCache"));
+			if (bForceUsePDBCache)
+			{
+				UE_LOG(LogCrashDebugHelper, Log, TEXT("No cached symbols available. Exiting due to -ForceUsePDBCache."));
+				return false;
+			}
+		}
 	}
 
 	// Initialise the symbol options
diff --git a/Engine/Source/Developer/CrashDebugHelper/Private/Windows/WindowsPlatformStackWalkExt.cpp b/Engine/Source/Developer/CrashDebugHelper/Private/Windows/WindowsPlatformStackWalkExt.cpp
--- a/Engine/Source/Developer/CrashDebugHelper/Private/Windows/WindowsPlatformStackWalkExt.cpp
+++ b/Engine/Source/Developer/CrashDebugHelper/Private/Windows/WindowsPlatformStackWalkExt.cpp
@@ -459,7 +459,9 @@ int FWindowsPlatformStackWalkExt::GetCallstacks(bool bTrimCallstack)
 	bool bFoundSourceFile = false;
 	void* ContextData = FMemStack::Get().PushBytes( MaxFramesSize, 0 );
 	FMemory::Memzero( ContextData, MaxFramesSize );
+	UE_LOG(LogCrashDebugHelper, Log, TEXT("Running GetContextStackTrace()"));
 	HRESULT HR = Control->GetContextStackTrace( Context, ContextUsed, StackFrames, MaxFrames, ContextData, MaxFramesSize, ContextUsed, &Count );
+	UE_LOG(LogCrashDebugHelper, Log, TEXT("GetContextStackTrace() got %d frames"), Count);
 
 	int32 AssertOrEnsureIndex = -1;
 
diff --git a/Engine/Source/Developer/CrashDebugHelper/Public/CrashDebugHelper.h b/Engine/Source/Developer/CrashDebugHelper/Public/CrashDebugHelper.h
--- a/Engine/Source/Developer/CrashDebugHelper/Public/CrashDebugHelper.h
+++ b/Engine/Source/Developer/CrashDebugHelper/Public/CrashDebugHelper.h
@@ -293,8 +293,12 @@ public:
 
 	/**
 	 * Sync the branch root relative file names to the requested label
+	 *
+	 *	@param	bOutPDBCacheEntryValid		Returns whether the PDB cache entry was found or created and whether it contains files.
+	 *
+	 *	@return	bool						true if successful, false if not
 	 */
-	virtual bool SyncModules();
+	virtual bool SyncModules(bool& bOutPDBCacheEntryValid);
 
 	/**
 	 *	Sync a single source file to the requested CL.
