Crash server PDB cache switched to indexing builds by BUILD_VERSION
[UE-34884] - Builds with the same engine version sometimes fail to symbolicate

diff --git a/Engine/Source/Developer/CrashDebugHelper/Private/CrashDebugHelper.cpp b/Engine/Source/Developer/CrashDebugHelper/Private/CrashDebugHelper.cpp
--- a/Engine/Source/Developer/CrashDebugHelper/Private/CrashDebugHelper.cpp
+++ b/Engine/Source/Developer/CrashDebugHelper/Private/CrashDebugHelper.cpp
@@ -309,13 +309,12 @@ bool ICrashDebugHelper::SyncModules(bool& bOutPDBCacheEntryValid)
 	const bool bHasExecutable = !CrashInfo.ExecutablesPath.IsEmpty();
 	const bool bHasSymbols = !CrashInfo.SymbolsPath.IsEmpty();
 	
-	const bool bContainsProductVersion = FPDBCache::Get().ContainsPDBCacheEntry( CrashInfo.EngineVersion );
 	if( bHasExecutable && bHasSymbols )
 	{
-		if( bContainsProductVersion )
+		if(FPDBCache::Get().ContainsPDBCacheEntry(CrashInfo.BuildVersion))
 		{
-			UE_LOG( LogCrashDebugHelper, Warning, TEXT( "Using cached storage: %s" ), *CrashInfo.EngineVersion );
-			CrashInfo.PDBCacheEntry = FPDBCache::Get().FindAndTouchPDBCacheEntry( CrashInfo.EngineVersion );
+			UE_LOG( LogCrashDebugHelper, Warning, TEXT( "Using cached storage: %s" ), *CrashInfo.BuildVersion);
+			CrashInfo.PDBCacheEntry = FPDBCache::Get().FindAndTouchPDBCacheEntry( CrashInfo.BuildVersion);
 		}
 		else
 		{
@@ -354,7 +353,7 @@ bool ICrashDebugHelper::SyncModules(bool& bOutPDBCacheEntryValid)
 			}
 
 			// Initialize and add a new PDB Cache entry to the database.
-			CrashInfo.PDBCacheEntry = FPDBCache::Get().CreateAndAddPDBCacheEntryMixed( CrashInfo.EngineVersion, FilesToBeCached );
+			CrashInfo.PDBCacheEntry = FPDBCache::Get().CreateAndAddPDBCacheEntryMixed( CrashInfo.BuildVersion, FilesToBeCached );
 		}
 	}
 	// OBSOLETE PATH
@@ -377,15 +376,13 @@ bool ICrashDebugHelper::SyncModules(bool& bOutPDBCacheEntryValid)
 			TSharedRef<ISourceControlLabel> Label = Labels[0];
 			TSet<FString> FilesToSync;
 
-			// Use product version instead of label name to make a distinguish between chosen methods.
-			const bool bContainsLabelName = FPDBCache::Get().ContainsPDBCacheEntry(CrashInfo.LabelName);
-
-			if (bContainsProductVersion)
+			if (FPDBCache::Get().ContainsPDBCacheEntry(CrashInfo.EngineVersion))
 			{
 				UE_LOG(LogCrashDebugHelper, Warning, TEXT("Using cached storage: %s"), *CrashInfo.EngineVersion);
 				CrashInfo.PDBCacheEntry = FPDBCache::Get().FindAndTouchPDBCacheEntry(CrashInfo.EngineVersion);
 			}
-			else if (bContainsLabelName)
+			// Use product version instead of label name to make a distinguish between chosen methods.
+			else if (FPDBCache::Get().ContainsPDBCacheEntry(CrashInfo.LabelName))
 			{
 				UE_LOG(LogCrashDebugHelper, Warning, TEXT("Using cached storage: %s"), *CrashInfo.LabelName);
 				CrashInfo.PDBCacheEntry = FPDBCache::Get().FindAndTouchPDBCacheEntry(CrashInfo.LabelName);
