Update UT builds path in config for MDD on Crash Report server.
Also, increase MDD lock timeout for larger pdb cache folders.

diff --git a/Engine/Source/Developer/CrashDebugHelper/Private/CrashDebugHelper.cpp b/Engine/Source/Developer/CrashDebugHelper/Private/CrashDebugHelper.cpp
--- a/Engine/Source/Developer/CrashDebugHelper/Private/CrashDebugHelper.cpp
+++ b/Engine/Source/Developer/CrashDebugHelper/Private/CrashDebugHelper.cpp
@@ -222,7 +222,7 @@ bool ICrashDebugHelper::Init()
 		if (CrashInfo.bMutexPDBCache && !CrashInfo.PDBCacheLockName.IsEmpty())
 		{
 			// Scoped lock
-			FSystemWideCriticalSection PDBCacheLock(CrashInfo.PDBCacheLockName, FTimespan(0, 0, 3, 0, 0));
+			FSystemWideCriticalSection PDBCacheLock(CrashInfo.PDBCacheLockName, FTimespan(0, 0, 10, 0, 0));
 			if (PDBCacheLock.IsValid())
 			{
 				FPDBCache::Get().Init();
diff --git a/Engine/Source/Developer/CrashDebugHelper/Private/Windows/CrashDebugHelperWindows.cpp b/Engine/Source/Developer/CrashDebugHelper/Private/Windows/CrashDebugHelperWindows.cpp
--- a/Engine/Source/Developer/CrashDebugHelper/Private/Windows/CrashDebugHelperWindows.cpp
+++ b/Engine/Source/Developer/CrashDebugHelper/Private/Windows/CrashDebugHelperWindows.cpp
@@ -41,7 +41,7 @@ bool FCrashDebugHelperWindows::CreateMinidumpDiagnosticReport( const FString& In
 			{
 				// Scoped lock
 				UE_LOG(LogCrashDebugHelper, Log, TEXT("Locking for InitSymbols()"));
-				FSystemWideCriticalSection PDBCacheLock(CrashInfo.PDBCacheLockName, FTimespan(0, 0, 3, 0, 0));
+				FSystemWideCriticalSection PDBCacheLock(CrashInfo.PDBCacheLockName, FTimespan(0, 0, 10, 0, 0));
 				if (PDBCacheLock.IsValid())
 				{
 					bInitSymbols = InitSymbols(WindowsStackWalkExt, bSyncSymbols);
