Static analysis fixes.

diff --git a/Engine/Source/Developer/CrashDebugHelper/Private/Windows/WindowsPlatformStackWalkExt.cpp b/Engine/Source/Developer/CrashDebugHelper/Private/Windows/WindowsPlatformStackWalkExt.cpp
--- a/Engine/Source/Developer/CrashDebugHelper/Private/Windows/WindowsPlatformStackWalkExt.cpp
+++ b/Engine/Source/Developer/CrashDebugHelper/Private/Windows/WindowsPlatformStackWalkExt.cpp
@@ -415,7 +415,7 @@ void FWindowsPlatformStackWalkExt::GetExceptionInfo()
 	ULONG ProcessID = 0;
 	ULONG ThreadId = 0;
 	TCHAR Description[MAX_PATH] = { 0 };
-	Control->GetLastEventInformationWide( &ExceptionType, &ProcessID, &ThreadId, NULL, 0, NULL, Description, sizeof( Description ), NULL );
+	Control->GetLastEventInformationWide( &ExceptionType, &ProcessID, &ThreadId, NULL, 0, NULL, Description, ARRAY_COUNT( Description ), NULL );
 
 	Exception.Code = ExceptionType;
 	Exception.ProcessId = ProcessID;
diff --git a/Engine/Source/Runtime/Analytics/QoSReporter/Private/QoSReporter.cpp b/Engine/Source/Runtime/Analytics/QoSReporter/Private/QoSReporter.cpp
--- a/Engine/Source/Runtime/Analytics/QoSReporter/Private/QoSReporter.cpp
+++ b/Engine/Source/Runtime/Analytics/QoSReporter/Private/QoSReporter.cpp
@@ -220,7 +220,8 @@ void FQoSReporter::Tick()
 	}
 
 	// detect too long pauses between ticks, unless configured to ignore them or running under debugger
-	if (!QOS_IGNORE_HITCHES && bCountHitches && !FPlatformMisc::IsDebuggerPresent())
+#if !QOS_IGNORE_HITCHES
+	if (bCountHitches && !FPlatformMisc::IsDebuggerPresent())
 	{
 		const double DeltaBetweenTicks = CurrentTime - PreviousTickTime;
 
@@ -250,6 +251,7 @@ void FQoSReporter::Tick()
 				DeltaBetweenTicks, GAverageFPS);
 		}
 	}
+#endif
 
 	PreviousTickTime = CurrentTime;
 }
