CRC: Fix crash callstacks including more lines than necessary when CRC symbolicates from a local minidump. Now skips the same number of frames as the portable callstack.

#rnx

diff --git a/Engine/Source/Developer/CrashDebugHelper/Private/Windows/CrashDebugHelperWindows.cpp b/Engine/Source/Developer/CrashDebugHelper/Private/Windows/CrashDebugHelperWindows.cpp
--- a/Engine/Source/Developer/CrashDebugHelper/Private/Windows/CrashDebugHelperWindows.cpp
+++ b/Engine/Source/Developer/CrashDebugHelper/Private/Windows/CrashDebugHelperWindows.cpp
@@ -29,7 +29,7 @@ bool FCrashDebugHelperWindows::CreateMinidumpDiagnosticReport( const FString& In
 
 	const bool bReady = WindowsStackWalkExt.InitStackWalking();
 
-	bool bHasAtLeastThreeValidFunctions = false;
+	bool bResult = false;
 	if( bReady && WindowsStackWalkExt.OpenDumpFile( InCrashDumpFilename ) )
 	{
 		if (CrashInfo.BuiltFromCL != FCrashInfo::INVALID_CHANGELIST)
@@ -71,7 +71,7 @@ bool FCrashDebugHelperWindows::CreateMinidumpDiagnosticReport( const FString& In
 				WindowsStackWalkExt.GetExceptionInfo();
 
 				// Get the callstacks for each thread
-				bHasAtLeastThreeValidFunctions = WindowsStackWalkExt.GetCallstacks(!bNoTrim) >= 3;
+				WindowsStackWalkExt.GetCallstacks();
 
 				// Sync the source file where the crash occurred
 				if (CrashInfo.SourceFile.Len() > 0)
@@ -97,6 +97,9 @@ bool FCrashDebugHelperWindows::CreateMinidumpDiagnosticReport( const FString& In
 						SyncAndReadSourceFile(bSyncSymbols, bAnnotate, CrashInfo.BuiltFromCL);
 					}
 				}
+
+				// Set the result
+				bResult = true;
 			}
 			else
 			{
@@ -118,7 +121,7 @@ bool FCrashDebugHelperWindows::CreateMinidumpDiagnosticReport( const FString& In
 		ShutdownSourceControl();
 	}
 
-	return bHasAtLeastThreeValidFunctions;
+	return bResult;
 }
 
 bool FCrashDebugHelperWindows::InitSymbols(FWindowsPlatformStackWalkExt& WindowsStackWalkExt, bool bSyncSymbols)
diff --git a/Engine/Source/Developer/CrashDebugHelper/Private/Windows/WindowsPlatformStackWalkExt.cpp b/Engine/Source/Developer/CrashDebugHelper/Private/Windows/WindowsPlatformStackWalkExt.cpp
--- a/Engine/Source/Developer/CrashDebugHelper/Private/Windows/WindowsPlatformStackWalkExt.cpp
+++ b/Engine/Source/Developer/CrashDebugHelper/Private/Windows/WindowsPlatformStackWalkExt.cpp
@@ -424,10 +424,9 @@ void FWindowsPlatformStackWalkExt::GetExceptionInfo()
 }
 
 
-int FWindowsPlatformStackWalkExt::GetCallstacks(bool bTrimCallstack)
+void FWindowsPlatformStackWalkExt::GetCallstacks()
 {
-	const int32 MAX_NAME_LENGHT = FProgramCounterSymbolInfo::MAX_NAME_LENGTH;
-	int32 NumValidFunctionNames = 0;
+	const int32 MAX_NAME_LENGTH = FProgramCounterSymbolInfo::MAX_NAME_LENGTH;
 
 	FCrashExceptionInfo& Exception = CrashInfo.Exception;
 
@@ -445,7 +444,7 @@ int FWindowsPlatformStackWalkExt::GetCallstacks(bool bTrimCallstack)
 	HRESULT hr = Control->GetStoredEventInformation(&DebugEvent, &ProcessID, &ThreadID, Context, ContextSize, &ContextUsed, NULL, 0, 0);
 	if( FAILED(hr) )
 	{
-		return NumValidFunctionNames;
+		return;
 	}
 
 	// Some magic number checks
@@ -510,8 +509,6 @@ int FWindowsPlatformStackWalkExt::GetCallstacks(bool bTrimCallstack)
 			// https://msdn.microsoft.com/en-us/library/windows/hardware/ff547186(v=vs.85).aspx
 			if( ModuleAndFunction.Contains( TEXT( "!" ) ) )
 			{
-				NumValidFunctionNames++;
- 
 				ModuleAndFunction.Split( TEXT( "!" ), &ModuleName, &FunctionName );
 				FunctionName += TEXT( "()" );
 			}
@@ -520,18 +517,6 @@ int FWindowsPlatformStackWalkExt::GetCallstacks(bool bTrimCallstack)
 				ModuleName = ModuleAndFunction;
 			}
 
-			// #CrashReport: 2015-07-24 Add this to other platforms
-			// If we find an assert, the actual source file we're interested in is the next one up, so reset the source file found flag
-			if( FunctionName.Len() > 0 )
-			{
-				if( FunctionName.Contains( TEXT( "FDebug::" ), ESearchCase::CaseSensitive )
-					|| FunctionName.Contains( TEXT( "NewReportEnsure" ), ESearchCase::CaseSensitive ) )
-				{
-					bFoundSourceFile = false;
-					AssertOrEnsureIndex = Exception.CallStackString.Num();
-				}
-			}
-
 			// FString InModuleName, FString InFunctionName, FString InFilename, uint32 InLineNumber, uint64 InSymbolDisplacement, uint64 InOffsetInModule, uint64 InProgramCounter
 			FProgramCounterSymbolInfoEx SymbolInfo( ModuleName, FunctionName, SourceName, LineNumber, Displacement, Offset, 0 );
 			FString GenericFormattedCallstackLine;
@@ -541,17 +526,6 @@ int FWindowsPlatformStackWalkExt::GetCallstacks(bool bTrimCallstack)
 			UE_LOG( LogCrashDebugHelper, Log, TEXT( "%3u: %s" ), StackIndex, *GenericFormattedCallstackLine );
 		}
 	}
-
-	// Remove callstack entries below FDebug, we don't need them.
-	if (bTrimCallstack && AssertOrEnsureIndex > 0)
-	{	
-		Exception.CallStackString.RemoveAt( 0, AssertOrEnsureIndex );
-		UE_LOG( LogCrashDebugHelper, Warning, TEXT( "Callstack trimmed to %i entries" ), Exception.CallStackString.Num() );
-	}
-
-	UE_LOG( LogCrashDebugHelper, Warning, TEXT( "Callstack generated with %i valid function names" ), NumValidFunctionNames );
-
-	return NumValidFunctionNames;
 }
 
 
diff --git a/Engine/Source/Developer/CrashDebugHelper/Private/Windows/WindowsPlatformStackWalkExt.h b/Engine/Source/Developer/CrashDebugHelper/Private/Windows/WindowsPlatformStackWalkExt.h
--- a/Engine/Source/Developer/CrashDebugHelper/Private/Windows/WindowsPlatformStackWalkExt.h
+++ b/Engine/Source/Developer/CrashDebugHelper/Private/Windows/WindowsPlatformStackWalkExt.h
@@ -51,13 +51,8 @@ struct FWindowsPlatformStackWalkExt
 	/** Extracts info about the exception that caused the crash. */
 	void GetExceptionInfo();
 
-	/**
-	 * Gets the callstack of the crash.
-	 *
-	 * @param bTrimCallstack	If true, trims what it thinks are irrelevant entries after a debug or assert. If false, leaves all callstack entries. 
-	 * @return the number of valid function names
-	 */
-	int32 GetCallstacks(bool bTrimCallstack);
+	/** Gets the callstack of the crash. */
+	void GetCallstacks();
 
 	/** Opens a minidump as a new session. */
 	bool OpenDumpFile( const FString& DumpFileName );
diff --git a/Engine/Source/Programs/CrashReportClient/Private/CrashDescription.cpp b/Engine/Source/Programs/CrashReportClient/Private/CrashDescription.cpp
--- a/Engine/Source/Programs/CrashReportClient/Private/CrashDescription.cpp
+++ b/Engine/Source/Programs/CrashReportClient/Private/CrashDescription.cpp
@@ -128,6 +128,7 @@ FPrimaryCrashProperties::FPrimaryCrashProperties()
 	bHasMiniDumpFile = false;
 	bHasLogFile = false;
 	bHasPrimaryData = false;
+	NumMinidumpFramesToIgnore = 0;
 }
 
 void FPrimaryCrashProperties::Shutdown()
@@ -441,6 +442,7 @@ FCrashContext::FCrashContext( const FString& CrashContextFilepath )
 		GetCrashProperty( DeploymentName, FGenericCrashContext::RuntimePropertiesTag, TEXT( "DeploymentName" ) );
 		GetCrashProperty( AppDefaultLocale, FGenericCrashContext::RuntimePropertiesTag, TEXT( "AppDefaultLocale" ) );
 		GetCrashProperty( bIsEnsure, FGenericCrashContext::RuntimePropertiesTag, TEXT("IsEnsure"));
+		GetCrashProperty( NumMinidumpFramesToIgnore, FGenericCrashContext::RuntimePropertiesTag, TEXT("NumMinidumpFramesToIgnore"));
 
 		if (CrashDumpMode == ECrashDumpMode::FullDump)
 		{
diff --git a/Engine/Source/Programs/CrashReportClient/Private/CrashDescription.h b/Engine/Source/Programs/CrashReportClient/Private/CrashDescription.h
--- a/Engine/Source/Programs/CrashReportClient/Private/CrashDescription.h
+++ b/Engine/Source/Programs/CrashReportClient/Private/CrashDescription.h
@@ -265,6 +265,11 @@ struct FPrimaryCrashProperties
 	 */
 	FCrashProperty GameSessionID;
 
+	/**
+	 * Specifies the number of stack frames in the callstack to ignore when symbolicating from a minidump.
+	 */
+	int32 NumMinidumpFramesToIgnore;
+
 	/**
 	 * An array of FStrings representing the callstack of the crash.
 	 * @RawCallStack	varchar(MAX)
@@ -356,7 +361,9 @@ struct FPrimaryCrashProperties
 	/** Whether this crash contains primary usable data. */
 	bool bHasPrimaryData;
 
-	/** Copy of CommandLine that isn't anonymized so it can be used to restart the process */
+	/** 
+	 * Copy of CommandLine that isn't anonymized so it can be used to restart the process 
+	 */
 	FString RestartCommandLine;
 
 	/**
diff --git a/Engine/Source/Programs/CrashReportClient/Private/GenericErrorReport.cpp b/Engine/Source/Programs/CrashReportClient/Private/GenericErrorReport.cpp
--- a/Engine/Source/Programs/CrashReportClient/Private/GenericErrorReport.cpp
+++ b/Engine/Source/Programs/CrashReportClient/Private/GenericErrorReport.cpp
@@ -176,7 +176,16 @@ void FGenericErrorReport::SetPrimaryCrashProperties( FPrimaryCrashProperties& ou
 	ICrashDebugHelper* Helper = CrashHelperModule.Get();
 	if (Helper && bValidCallstack)
 	{
-		out_PrimaryCrashProperties.CallStack = Helper->CrashInfo.Exception.CallStackString;
+		TArray<FString> CallStack = Helper->CrashInfo.Exception.CallStackString;
+
+		// Get the callstack and remove any frames that we don't care about
+		int64 NumMinidumpFramesToIgnore = out_PrimaryCrashProperties.NumMinidumpFramesToIgnore;
+		if (NumMinidumpFramesToIgnore > 0)
+		{
+			CallStack.RemoveAt(0, FMath::Min(CallStack.Num(), (int32)NumMinidumpFramesToIgnore));
+		}
+
+		out_PrimaryCrashProperties.CallStack = CallStack;
 		out_PrimaryCrashProperties.Modules = Helper->CrashInfo.ModuleNames;
 		out_PrimaryCrashProperties.SourceContext = Helper->CrashInfo.SourceContext;
 
diff --git a/Engine/Source/Runtime/Core/Private/GenericPlatform/GenericPlatformCrashContext.cpp b/Engine/Source/Runtime/Core/Private/GenericPlatform/GenericPlatformCrashContext.cpp
--- a/Engine/Source/Runtime/Core/Private/GenericPlatform/GenericPlatformCrashContext.cpp
+++ b/Engine/Source/Runtime/Core/Private/GenericPlatform/GenericPlatformCrashContext.cpp
@@ -257,6 +257,7 @@ void FGenericCrashContext::UpdateLocalizedStrings()
 
 FGenericCrashContext::FGenericCrashContext()
 	: bIsEnsure(false)
+	, NumMinidumpFramesToIgnore(0)
 {
 	CommonBuffer.Reserve( 32768 );
 	CrashContextIndex = StaticCrashContextIndex++;
@@ -324,7 +325,8 @@ void FGenericCrashContext::SerializeContentToBuffer() const
 	AddCrashProperty( TEXT( "EpicAccountId" ), *NCachedCrashContextProperties::EpicAccountId );
 
 	// Legacy callstack element for current crash reporter
-	AddCrashProperty(TEXT("CallStack"), TEXT(""));
+	AddCrashProperty( TEXT( "NumMinidumpFramesToIgnore"), NumMinidumpFramesToIgnore );
+	AddCrashProperty( TEXT( "CallStack" ), TEXT("") );
 
 	// Add new portable callstack element with crash stack
 	AddPortableCallStack();
@@ -405,6 +407,11 @@ void FGenericCrashContext::SerializeContentToBuffer() const
 	AddFooter();
 }
 
+void FGenericCrashContext::SetNumMinidumpFramesToIgnore(int InNumMinidumpFramesToIgnore)
+{
+	NumMinidumpFramesToIgnore = InNumMinidumpFramesToIgnore;
+}
+
 void FGenericCrashContext::SetDeploymentName(const FString& EpicApp)
 {
 	NCachedCrashContextProperties::DeploymentName = EpicApp;
@@ -647,10 +654,10 @@ FORCENOINLINE void FGenericCrashContext::CapturePortableCallStack(int32 NumStack
 
 	const int32 MaxDepth = 100;
 	TArray<FProgramCounterSymbolInfo> Stack = FPlatformStackWalk::GetStack(NumStackFramesToIgnore, MaxDepth, Context);
-	return SetPortableCallStack(NumStackFramesToIgnore, Stack);
+	return SetPortableCallStack(Stack);
 }
 
-void FGenericCrashContext::SetPortableCallStack(int32 NumStackFramesToIgnore, const TArray<FProgramCounterSymbolInfo>& Stack)
+void FGenericCrashContext::SetPortableCallStack(const TArray<FProgramCounterSymbolInfo>& Stack)
 {
 	uint32 ModuleEntries = (uint32)FPlatformStackWalk::GetProcessModuleCount();
 
diff --git a/Engine/Source/Runtime/Core/Private/Mac/MacPlatformMisc.cpp b/Engine/Source/Runtime/Core/Private/Mac/MacPlatformMisc.cpp
--- a/Engine/Source/Runtime/Core/Private/Mac/MacPlatformMisc.cpp
+++ b/Engine/Source/Runtime/Core/Private/Mac/MacPlatformMisc.cpp
@@ -1886,7 +1886,7 @@ void ReportHang(const TCHAR* ErrorMessage, const TArray<FProgramCounterSymbolInf
 
 		const bool bIsEnsure = true;
 		FMacCrashContext EnsureContext(bIsEnsure);
-		EnsureContext.SetPortableCallStack(0, Stack);
+		EnsureContext.SetPortableCallStack(Stack);
 		EnsureContext.GenerateEnsureInfoAndLaunchReporter();
 
 		bReentranceGuard = false;
diff --git a/Engine/Source/Runtime/Core/Private/Unix/UnixPlatformStackWalk.cpp b/Engine/Source/Runtime/Core/Private/Unix/UnixPlatformStackWalk.cpp
--- a/Engine/Source/Runtime/Core/Private/Unix/UnixPlatformStackWalk.cpp
+++ b/Engine/Source/Runtime/Core/Private/Unix/UnixPlatformStackWalk.cpp
@@ -644,7 +644,7 @@ void ReportHang(const TCHAR* ErrorMessage, const TArray<FProgramCounterSymbolInf
 
 		const bool bIsEnsure = true;
 		FUnixCrashContext EnsureContext(bIsEnsure);
-		EnsureContext.SetPortableCallStack(0, Stack);
+		EnsureContext.SetPortableCallStack(Stack);
 		EnsureContext.GenerateCrashInfoAndLaunchReporter(true);
 
 		bReentranceGuard = false;
diff --git a/Engine/Source/Runtime/Core/Private/Windows/WindowsPlatformCrashContext.cpp b/Engine/Source/Runtime/Core/Private/Windows/WindowsPlatformCrashContext.cpp
--- a/Engine/Source/Runtime/Core/Private/Windows/WindowsPlatformCrashContext.cpp
+++ b/Engine/Source/Runtime/Core/Private/Windows/WindowsPlatformCrashContext.cpp
@@ -318,7 +318,7 @@ void ReportHang(const TCHAR* ErrorMessage, const TArray<FProgramCounterSymbolInf
 	const bool bIsEnsure = true;
 
 	FWindowsPlatformCrashContext CrashContext(bIsEnsure);
-	CrashContext.SetPortableCallStack(0, Stack);
+	CrashContext.SetPortableCallStack(Stack);
 
 	EErrorReportUI ReportUI = IsInteractiveEnsureMode() ? EErrorReportUI::ShowDialog : EErrorReportUI::ReportInUnattendedMode;
 	ReportCrashUsingCrashReportClient(CrashContext, nullptr, ErrorMessage, ReportUI, bIsEnsure);
@@ -568,6 +568,9 @@ private:
 		const int32 IgnoreCount = FDebug::HasAsserted() ? 4 : 0;
 		CrashContext.CapturePortableCallStack(IgnoreCount, ContextWrapper);
 
+		// Also mark the same number of frames to be ignored if we symbolicate from the minidump
+		CrashContext.SetNumMinidumpFramesToIgnore(IgnoreCount);
+
 		// First launch the crash reporter client.
 #if WINVER > 0x502	// Windows Error Reporting is not supported on Windows XP
 		if (GUseCrashReportClient)
diff --git a/Engine/Source/Runtime/Core/Public/GenericPlatform/GenericPlatformCrashContext.h b/Engine/Source/Runtime/Core/Public/GenericPlatform/GenericPlatformCrashContext.h
--- a/Engine/Source/Runtime/Core/Public/GenericPlatform/GenericPlatformCrashContext.h
+++ b/Engine/Source/Runtime/Core/Public/GenericPlatform/GenericPlatformCrashContext.h
@@ -204,12 +204,15 @@ public:
 
 	/** Adds a plugin descriptor string to the enabled plugins list in the crash context */
 	static void AddPlugin(const FString& PluginDesc);
-	
+
+	/** Sets the number of stack frames to ignore when symbolicating from a minidump */
+	void SetNumMinidumpFramesToIgnore(int32 InNumMinidumpFramesToIgnore);
+
 	/** Generate raw call stack for crash report (image base + offset) */
 	void CapturePortableCallStack(int32 NumStackFramesToIgnore, void* Context);
 	
 	/** Sets the portable callstack to a specified stack */
-	void SetPortableCallStack(int32 NumStackFramesToIgnore, const TArray<FProgramCounterSymbolInfo>& Stack);
+	void SetPortableCallStack(const TArray<FProgramCounterSymbolInfo>& Stack);
 
 	/**
 	 * @return whether this crash is a non-crash event
@@ -223,6 +226,7 @@ public:
 
 protected:
 	bool bIsEnsure;
+	int NumMinidumpFramesToIgnore;
 	TArray<FCrashStackFrame> CallStack;
 
 private:
