Merge 8685846 from //UE4/Dev-Rendering to Dev-Core (//UE4/Dev-Core)

---

[macOS, WebBrowser] Use a 2-byte wide character representation on macOS.

Convert the wide-character representation used on macOS to match that of all the other supported desktop and mobile platforms.

<see original CL for reviewers, etc.>
---

#rnx

diff --git a/Engine/Build/BatchFiles/Mac/GenerateLLDBInit.sh b/Engine/Build/BatchFiles/Mac/GenerateLLDBInit.sh
--- a/Engine/Build/BatchFiles/Mac/GenerateLLDBInit.sh
+++ b/Engine/Build/BatchFiles/Mac/GenerateLLDBInit.sh
@@ -2,5 +2,5 @@
 
 if [ ! -f ~/.lldbinit ]; then
   echo 'echo -n "settings set target.inline-breakpoint-strategy always\n"' | bash > ~/.lldbinit
-  echo 'echo -n "command script import \"`pwd`/../../../Extras/LLDBDataFormatters/UE4DataFormatters.py\"\n"' | bash >> ~/.lldbinit
+  echo 'echo -n "command script import \"`pwd`/../../../Extras/LLDBDataFormatters/UE4DataFormatters_2ByteChars.py\"\n"' | bash >> ~/.lldbinit
 fi
diff --git a/Engine/Source/Developer/Apple/MetalShaderFormat/Private/MetalShaderCompiler.cpp b/Engine/Source/Developer/Apple/MetalShaderFormat/Private/MetalShaderCompiler.cpp
--- a/Engine/Source/Developer/Apple/MetalShaderFormat/Private/MetalShaderCompiler.cpp
+++ b/Engine/Source/Developer/Apple/MetalShaderFormat/Private/MetalShaderCompiler.cpp
@@ -36,6 +36,15 @@ THIRD_PARTY_INCLUDES_END
 // The Metal standard library extensions we need for UE4.
 #include "ue4_stdlib.h"
 
+#if !PLATFORM_WINDOWS
+#if PLATFORM_TCHAR_IS_CHAR16
+#define FP_TEXT_PASTE(x) L ## x
+#define WTEXT(x) FP_TEXT_PASTE(x)
+#else
+#define WTEXT TEXT
+#endif
+#endif
+
 DEFINE_LOG_CATEGORY_STATIC(LogMetalShaderCompiler, Log, All); 
 
 static FString	GRemoteBuildServerHost;
@@ -335,11 +344,11 @@ FString GetMetalCompilerVers(FString const& PlatformPath)
 	bOK = ExecRemoteProcess(*PlatformPath, TEXT("-v"), nullptr, &Result, &Result);
 	if (bOK && Result.Len() > 0)
 	{
-		TCHAR Buffer[256];
+		ANSICHAR Buffer[256];
 #if !PLATFORM_WINDOWS
-		if(swscanf(*Result, TEXT("Apple LLVM version %*ls (%ls)"), Buffer))
+		if(sscanf(TCHAR_TO_ANSI(*Result), "Apple LLVM version %*s (%s)", Buffer))
 #else
-		if(swscanf_s(*Result, TEXT("Apple LLVM version %*ls (%ls)"), Buffer, 256))
+		if(sscanf_s(TCHAR_TO_ANSI(*Result), "Apple LLVM version %*s (%s)", Buffer, 256))
 #endif
 		{
 			Result = (&Buffer[0]);
@@ -573,7 +582,7 @@ uint16 GetXcodeVersion(uint64& BuildVersion)
 				uint32 Patch = 0;
 				int32 NumResults = 0;
 	#if !PLATFORM_WINDOWS
-				NumResults = swscanf(*Result, TEXT("Xcode %u.%u.%u"), &Major, &Minor, &Patch);
+				NumResults = swscanf(TCHAR_TO_WCHAR(*Result), WTEXT("Xcode %u.%u.%u"), &Major, &Minor, &Patch);
 	#else
 				NumResults = swscanf_s(*Result, TEXT("Xcode %u.%u.%u"), &Major, &Minor, &Patch);
 	#endif
@@ -612,7 +621,7 @@ bool ChecksumRemoteFile(FString const& RemotePath, uint32* CRC, uint32* Len)
 	if (bOK)
 	{
 #if !PLATFORM_WINDOWS
-		if(swscanf(*Output, TEXT("%u %u"), CRC, Len) != 2)
+		if(swscanf(TCHAR_TO_WCHAR(*Output), WTEXT("%u %u"), CRC, Len) != 2)
 #else
 		if(swscanf_s(*Output, TEXT("%u %u"), CRC, Len) != 2)
 #endif
diff --git a/Engine/Source/Developer/CrashDebugHelper/Private/Mac/CrashDebugHelperMac.cpp b/Engine/Source/Developer/CrashDebugHelper/Private/Mac/CrashDebugHelperMac.cpp
--- a/Engine/Source/Developer/CrashDebugHelper/Private/Mac/CrashDebugHelperMac.cpp
+++ b/Engine/Source/Developer/CrashDebugHelper/Private/Mac/CrashDebugHelperMac.cpp
@@ -10,6 +10,13 @@
 #include "Misc/Paths.h"
 #include <cxxabi.h>
 
+#if PLATFORM_TCHAR_IS_CHAR16
+#define FP_TEXT_PASTE(x) L ## x
+#define WTEXT(x) FP_TEXT_PASTE(x)
+#else
+#define WTEXT TEXT
+#endif
+
 FString ExtractRelativePath( const TCHAR* BaseName, TCHAR const* FullName )
 {
 	FString FullPath = FString( FullName ).ToLower();
@@ -47,7 +54,7 @@ static int32 ParseReportVersion(TCHAR const* CrashLog, int32& OutVersion)
 	TCHAR const* VersionLine = FCStringWide::Strstr(CrashLog, TEXT("Report Version:"));
 	if (VersionLine)
 	{
-		Found = swscanf(VersionLine, TEXT("%*ls %*ls %d"), &OutVersion);
+		Found = swscanf(TCHAR_TO_WCHAR(VersionLine), WTEXT("%*ls %*ls %d"), &OutVersion);
 	}
 	return Found;
 }
@@ -58,10 +65,11 @@ static int32 ParseVersion(TCHAR const* CrashLog, int32& OutMajor, int32& OutMino
 	TCHAR const* VersionLine = FCStringWide::Strstr(CrashLog, TEXT("Version:"));
 	if (VersionLine)
 	{
-		TCHAR Branch[257] = {0};
-		Found = swscanf(VersionLine, TEXT("%*s %d.%d.%d (%*d.%*d.%*d-%d+%256ls)"), &OutMajor, &OutMinor, &OutBuild, &OutChangeList, Branch);
+		wchar_t RawBranch[257] = {0};
+		Found = swscanf(TCHAR_TO_WCHAR(VersionLine), WTEXT("%*s %d.%d.%d (%*d.%*d.%*d-%d+%256ls)"), &OutMajor, &OutMinor, &OutBuild, &OutChangeList, RawBranch);
 		if(Found == 5)
 		{
+			TCHAR* Branch = WCHAR_TO_TCHAR(RawBranch);
 			TCHAR* BranchEnd = FCStringWide::Strchr(Branch, TEXT(')'));
 			if(BranchEnd)
 			{
@@ -79,11 +87,11 @@ static int32 ParseOS(TCHAR const* CrashLog, uint16& OutMajor, uint16& OutMinor,
 	TCHAR const* VersionLine = FCStringWide::Strstr(CrashLog, TEXT("OS Version:"));
 	if (VersionLine)
 	{
-		Found = swscanf(VersionLine, TEXT("%*s %*s Mac OS X %hd.%hd.%hd (%hxd)"), &OutMajor, &OutMinor, &OutPatch, &OutBuild);
+		Found = swscanf(TCHAR_TO_WCHAR(VersionLine), WTEXT("%*s %*s Mac OS X %hd.%hd.%hd (%hxd)"), &OutMajor, &OutMinor, &OutPatch, &OutBuild);
         if ( Found == 2 )
         {
             OutPatch = 0;
-            Found += swscanf(VersionLine, TEXT("%*s %*s Mac OS X %*hd.%*hd (%hxd)"), &OutBuild) + 1;
+            Found += swscanf(TCHAR_TO_WCHAR(VersionLine), WTEXT("%*s %*s Mac OS X %*hd.%*hd (%hxd)"), &OutBuild) + 1;
         }
 	}
 	return Found;
@@ -205,14 +213,16 @@ static int32 ParseExceptionCode(TCHAR const* CrashLog, uint32& OutExceptionCode)
 	TCHAR const* Line = FCStringWide::Strstr(CrashLog, TEXT("Exception Type:"));
 	if(Line)
 	{
-		TCHAR Buffer[257] = {0};
-		Found = swscanf(Line, TEXT("%*s %*s %*s (%256ls)"), Buffer);
+		wchar_t RawBuffer[257] = {0};
+		Found = swscanf(TCHAR_TO_WCHAR(Line), WTEXT("%*s %*s %*s (%256ls)"), RawBuffer);
 		if(!Found)
 		{
-			Found = swscanf(Line, TEXT("%*s %*s %256ls"), Buffer);
+			Found = swscanf(TCHAR_TO_WCHAR(Line), WTEXT("%*s %*s %256ls"), RawBuffer);
 		}
 		if(Found)
 		{
+			TCHAR* Buffer = WCHAR_TO_TCHAR(RawBuffer);
+
 			TCHAR* End = FCStringWide::Strchr(Buffer, TEXT(')'));
 			if(End)
 			{
@@ -256,7 +266,7 @@ static int32 ParseExceptionCode(TCHAR const* CrashLog, uint32& OutExceptionCode)
 			}
 			else if(FString(Buffer).IsNumeric())
 			{
-				Found = swscanf(Buffer, TEXT("%u"), &OutExceptionCode);
+				Found = swscanf(TCHAR_TO_WCHAR(Buffer), WTEXT("%u"), &OutExceptionCode);
 			}
 			else
 			{
@@ -274,7 +284,7 @@ static int32 ParseCrashedThread(TCHAR const* CrashLog, uint32& OutThreadNumber)
 	TCHAR const* Line = FCStringWide::Strstr(CrashLog, TEXT("Crashed Thread:"));
 	if (Line)
 	{
-		Found = swscanf(Line, TEXT("%*s %*s %u"), &OutThreadNumber);
+		Found = swscanf(TCHAR_TO_WCHAR(Line), WTEXT("%*s %*s %u"), &OutThreadNumber);
 	}
 	return Found;
 }
@@ -285,7 +295,7 @@ static int32 ParseProcessID(TCHAR const* CrashLog, uint32& OutPID)
 	TCHAR const* Line = FCStringWide::Strstr(CrashLog, TEXT("Process:"));
 	if (Line)
 	{
-		Found = swscanf(Line, TEXT("%*s %*s [%u]"), &OutPID);
+		Found = swscanf(TCHAR_TO_WCHAR(Line), WTEXT("%*s %*s [%u]"), &OutPID);
 	}
 	return Found;
 }
@@ -318,18 +328,18 @@ static TCHAR const* FindCrashedThreadStack(TCHAR const* CrashLog)
 
 static int32 ParseThreadStackLine(TCHAR const* StackLine, FString& OutModuleName, uint64& OutProgramCounter, FString& OutFunctionName, FString& OutFileName, int32& OutLineNumber)
 {
-	TCHAR ModuleName[257];
-	TCHAR FunctionName[1025];
-	TCHAR FileName[257];
+	wchar_t RawModuleName[257];
+	wchar_t RawFunctionName[1025];
+	wchar_t RawFileName[257];
 	
-	int32 Found = swscanf(StackLine, TEXT("%*d %256ls 0x%lx"), ModuleName, &OutProgramCounter);
+	int32 Found = swscanf(TCHAR_TO_WCHAR(StackLine), WTEXT("%*d %256ls 0x%lx"), RawModuleName, &OutProgramCounter);
 	if(Found == 2)
 	{
 		uint64 FunctionAddress = 0;
 		uint32 FunctionOffset = 0;
-		if(swscanf(StackLine, TEXT("%*d %*ls %*lx 0x%lx + %d"), &FunctionAddress, &FunctionOffset) == 0)
+		if(swscanf(TCHAR_TO_WCHAR(StackLine), WTEXT("%*d %*ls %*lx 0x%lx + %d"), &FunctionAddress, &FunctionOffset) == 0)
 		{
-			Found += swscanf(StackLine, TEXT("%*d %*ls %*lx %1024ls + %*d (%256ls:%d)"), FunctionName, FileName, &OutLineNumber);
+			Found += swscanf(TCHAR_TO_WCHAR(StackLine), WTEXT("%*d %*ls %*lx %1024ls + %*d (%256ls:%d)"), RawFunctionName, RawFileName, &OutLineNumber);
 		}
 	}
 	
@@ -338,11 +348,12 @@ static int32 ParseThreadStackLine(TCHAR const* StackLine, FString& OutModuleName
 		case 5:
 		case 4:
 		{
-			OutFileName = FileName;
+			OutFileName = WCHAR_TO_TCHAR(RawFileName);
 		}
 		case 3:
 		{
 			int32 Status = -1;
+			const TCHAR* FunctionName = WCHAR_TO_TCHAR(RawFunctionName);
 			ANSICHAR* DemangledName = abi::__cxa_demangle(TCHAR_TO_UTF8(FunctionName), nullptr, nullptr, &Status);
 			if(DemangledName && Status == 0)
 			{
@@ -363,7 +374,7 @@ static int32 ParseThreadStackLine(TCHAR const* StackLine, FString& OutModuleName
 		case 2:
 		case 1:
 		{
-			OutModuleName = ModuleName;
+			OutModuleName = WCHAR_TO_TCHAR(RawModuleName);
 		}
 		default:
 		{
@@ -436,12 +447,12 @@ static TCHAR const* FindModules(TCHAR const* CrashLog)
 static int32 ParseModuleVersion(TCHAR const* Version, uint16& OutMajor, uint16& OutMinor, uint16& OutPatch, uint16& OutBuild)
 {
 	OutMajor = OutMinor = OutPatch = OutBuild = 0;
-	int32 Found = swscanf(Version, TEXT("%hu.%hu.%hu"), &OutMajor, &OutMinor, &OutPatch);
+	int32 Found = swscanf(TCHAR_TO_WCHAR(Version), WTEXT("%hu.%hu.%hu"), &OutMajor, &OutMinor, &OutPatch);
 	TCHAR const* CurrentStart = FCStringWide::Strchr(Version, TEXT('-'));
 	if(CurrentStart)
 	{
 		int32 Components[3] = {0, 0, 0};
-		int32 Result = swscanf(CurrentStart, TEXT("%*ls %d.%d.%d"), &Components[0], &Components[1], &Components[2]);
+		int32 Result = swscanf(TCHAR_TO_WCHAR(CurrentStart), WTEXT("%*ls %d.%d.%d"), &Components[0], &Components[1], &Components[2]);
 		OutBuild = (uint16)(Components[0] * 10000) + (Components[1] * 100) + (Components[2]);
 		Found = 4;
 	}
@@ -451,11 +462,11 @@ static int32 ParseModuleVersion(TCHAR const* Version, uint16& OutMajor, uint16&
 static bool ParseModuleLine(TCHAR const* ModuleLine, FCrashModuleInfo& OutModule)
 {
 	bool bOK = false;
-	TCHAR ModuleName[257] = {0};
+	wchar_t RawModuleName[257] = {0};
 	uint64 ModuleBase = 0;
 	uint64 ModuleEnd = 0;
 	
-	int32 Found = swscanf(ModuleLine, TEXT("%lx %*ls %lx %256ls"), &ModuleBase, &ModuleEnd, ModuleName);
+	int32 Found = swscanf(TCHAR_TO_WCHAR(ModuleLine), WTEXT("%lx %*ls %lx %256ls"), &ModuleBase, &ModuleEnd, RawModuleName);
 	switch (Found)
 	{
 		case 3:
diff --git a/Engine/Source/Runtime/Core/Public/Mac/MacPlatform.h b/Engine/Source/Runtime/Core/Public/Mac/MacPlatform.h
--- a/Engine/Source/Runtime/Core/Public/Mac/MacPlatform.h
+++ b/Engine/Source/Runtime/Core/Public/Mac/MacPlatform.h
@@ -8,6 +8,8 @@
 
 #include "Clang/ClangPlatform.h"
 
+#define PLATFORM_MAC_USE_CHAR16 1
+
 /**
 * Mac specific types
 **/
@@ -16,7 +18,12 @@ struct FMacPlatformTypes : public FGenericPlatformTypes
 	typedef unsigned int		DWORD;
 	typedef size_t				SIZE_T;
 	typedef decltype(NULL)		TYPE_OF_NULL;
+#if PLATFORM_MAC_USE_CHAR16
+	typedef char16_t			WIDECHAR;
+	typedef WIDECHAR			TCHAR;
+#else
 	typedef char16_t			CHAR16;
+#endif
 };
 
 typedef FMacPlatformTypes FPlatformTypes;
@@ -37,7 +44,11 @@ typedef FMacPlatformTypes FPlatformTypes;
 //#define PLATFORM_USE_LS_SPEC_FOR_WIDECHAR			1
 #define PLATFORM_USE_SYSTEM_VSWPRINTF				0
 #define PLATFORM_COMPILER_DISTINGUISHES_INT_AND_LONG			1
+#if PLATFORM_MAC_USE_CHAR16
+#define PLATFORM_TCHAR_IS_CHAR16					1
+#else
 #define PLATFORM_TCHAR_IS_4_BYTES					1
+#endif
 #define PLATFORM_HAS_BSD_TIME						1
 //#define PLATFORM_USE_PTHREADS						1
 #define PLATFORM_MAX_FILEPATH_LENGTH_DEPRECATED		MAC_MAX_PATH
diff --git a/Engine/Source/Runtime/WebBrowser/Private/CEF/CEFImeHandler.cpp b/Engine/Source/Runtime/WebBrowser/Private/CEF/CEFImeHandler.cpp
--- a/Engine/Source/Runtime/WebBrowser/Private/CEF/CEFImeHandler.cpp
+++ b/Engine/Source/Runtime/WebBrowser/Private/CEF/CEFImeHandler.cpp
@@ -17,7 +17,7 @@ FCEFImeHandler::FCEFImeHandler(CefRefPtr<CefBrowser> Browser)
 bool FCEFImeHandler::OnProcessMessageReceived(CefRefPtr<CefBrowser> Browser, CefProcessId SourceProcess, CefRefPtr<CefProcessMessage> Message)
 {
 	bool Result = false;
-	FString MessageName = Message->GetName().ToWString().c_str();
+	FString MessageName = WCHAR_TO_TCHAR(Message->GetName().ToWString().c_str());
 	if (MessageName == TEXT("UE::IME::FocusChanged"))
 	{
 		Result = HandleFocusChangedMessage(Message->GetArgumentList());
@@ -135,8 +135,8 @@ bool FCEFImeHandler::HandleFocusChangedMessage(CefRefPtr<CefListValue> MessageAr
 		MessageArguments->GetType(6) == VTYPE_INT &&
 		MessageArguments->GetType(7) == VTYPE_INT)
 	{
-		FString Type = MessageArguments->GetString(0).ToWString().c_str();
-		FString Name = MessageArguments->GetString(1).ToWString().c_str();
+		FString Type = WCHAR_TO_TCHAR(MessageArguments->GetString(0).ToWString().c_str());
+		FString Name = WCHAR_TO_TCHAR(MessageArguments->GetString(1).ToWString().c_str());
 		bool bIsEditable = MessageArguments->GetBool(2);
 
 		if (Type == TEXT("DOM_NODE_TYPE_ELEMENT") && 
@@ -208,4 +208,4 @@ void FCEFImeHandler::CEFCompositionRangeChanged(const CefRange& SelectionRange,
 	}
 }
 
-#endif
\ No newline at end of file
+#endif
diff --git a/Engine/Source/Runtime/WebBrowser/Private/CEF/CEFTextInputMethodContext.cpp b/Engine/Source/Runtime/WebBrowser/Private/CEF/CEFTextInputMethodContext.cpp
--- a/Engine/Source/Runtime/WebBrowser/Private/CEF/CEFTextInputMethodContext.cpp
+++ b/Engine/Source/Runtime/WebBrowser/Private/CEF/CEFTextInputMethodContext.cpp
@@ -86,7 +86,7 @@ void FCEFTextInputMethodContext::SetSelectionRange(const uint32 BeginIndex, cons
 	SelectionRangeLength = Length;
 	SelectionCaretPosition = CaretPosition;
 
-	CefString Str = *CompositionString;
+	CefString Str = TCHAR_TO_WCHAR(*CompositionString);
 	std::vector<CefCompositionUnderline> underlines;
 	Owner->InternalCefBrowser->GetHost()->ImeSetComposition(
 		Str,
@@ -116,7 +116,7 @@ void FCEFTextInputMethodContext::SetTextInRange(const uint32 BeginIndex, const u
 	}
 	CompositionString = NewString;
 
-	CefString Str = *CompositionString;
+	CefString Str = TCHAR_TO_WCHAR(*CompositionString);
 	std::vector<CefCompositionUnderline> underlines;
 	Owner->InternalCefBrowser->GetHost()->ImeSetComposition(
 		Str,
@@ -241,7 +241,7 @@ void FCEFTextInputMethodContext::EndComposition()
 
 		if (CompositionString.Len() > 0)
 		{
-			CefString Result = *CompositionString;
+			CefString Result = TCHAR_TO_WCHAR(*CompositionString);
 			Owner->InternalCefBrowser->GetHost()->ImeCommitText(Result, CefRange(UINT32_MAX, UINT32_MAX), 0);
 		}
 		else
