Removed unneeded dependencies from CrashReportClient

Legacy code in CrashDebugHelper was dragging in SourceControl and AssetRegistry as dependencies for CrashReportClient.

This dependency is no longer needed, as internal crashes now always perform local symbolification (either via symbols built locally, or synced via UGS). Syncing symbols from Perforce or a network drive is no longer needed or used.

#rnx

diff --git a/Engine/Plugins/Developer/PlasticSourceControl/Source/PlasticSourceControl/Private/PlasticSourceControlRevision.cpp b/Engine/Plugins/Developer/PlasticSourceControl/Source/PlasticSourceControl/Private/PlasticSourceControlRevision.cpp
--- a/Engine/Plugins/Developer/PlasticSourceControl/Source/PlasticSourceControl/Private/PlasticSourceControlRevision.cpp
+++ b/Engine/Plugins/Developer/PlasticSourceControl/Source/PlasticSourceControl/Private/PlasticSourceControlRevision.cpp
@@ -52,18 +52,11 @@ bool FPlasticSourceControlRevision::Get( FString& InOutFilename ) const
 
 bool FPlasticSourceControlRevision::GetAnnotated( TArray<FAnnotationLine>& OutLines ) const
 {
-	// NOTE GetAnnotated: called only by SourceControlHelpers::AnnotateFile(),
-	//      called only by ICrashDebugHelper::AddAnnotatedSourceToReport() using a changelist/check identifier
-	//      called only by FCrashDebugHelperWindows::CreateMinidumpDiagnosticReport() (and Mac) to Extract annotated lines from a source file stored in Perforce, and add to the crash report.
-	//      called by - MinidumpDiagnosticsApp RunMinidumpDiagnostics() for Perfore ONLY "MinidumpDiagnostics.exe <Crash.dmp> [-Annotate] [-SyncSymbols] [-SyncMicrosoftSymbols]"
-	//                - FWindowsErrorReport::DiagnoseReport() (and Mac)
-	// Reserved for internal use by Epic Games with Perforce only
 	return false;
 }
 
 bool FPlasticSourceControlRevision::GetAnnotated( FString& InOutFilename ) const
 {
-	// NOTE: Unused, only the above method is called by the Editor
 	return false;
 }
 
diff --git a/Engine/Source/Developer/CrashDebugHelper/CrashDebugHelper.Build.cs b/Engine/Source/Developer/CrashDebugHelper/CrashDebugHelper.Build.cs
--- a/Engine/Source/Developer/CrashDebugHelper/CrashDebugHelper.Build.cs
+++ b/Engine/Source/Developer/CrashDebugHelper/CrashDebugHelper.Build.cs
@@ -22,7 +22,6 @@ public class CrashDebugHelper : ModuleRules
             PublicDependencyModuleNames.AddRange(
                 new string[] {
                 "Core",
-                "SourceControl"
                 }
             );
         }
diff --git a/Engine/Source/Developer/CrashDebugHelper/Private/CrashDebugHelper.cpp b/Engine/Source/Developer/CrashDebugHelper/Private/CrashDebugHelper.cpp
--- a/Engine/Source/Developer/CrashDebugHelper/Private/CrashDebugHelper.cpp
+++ b/Engine/Source/Developer/CrashDebugHelper/Private/CrashDebugHelper.cpp
@@ -10,157 +10,16 @@
 #include "Misc/ConfigCacheIni.h"
 #include "Misc/App.h"
 #include "CrashDebugHelperPrivate.h"
-#include "CrashDebugPDBCache.h"
-
 #include "Misc/EngineVersion.h"
-#if UE_EDITOR || IS_PROGRAM
-#include "ISourceControlOperation.h"
-#include "SourceControlOperations.h"
-#include "ISourceControlRevision.h"
-#include "ISourceControlProvider.h"
-#include "ISourceControlModule.h"
-#include "ISourceControlLabel.h"
-#include "SourceControlHelpers.h"
-#endif
-
-#ifndef MINIDUMPDIAGNOSTICS
-	#define MINIDUMPDIAGNOSTICS	0
-#endif
-
-/*-----------------------------------------------------------------------------
-	FCrashDebugHelperConfig
------------------------------------------------------------------------------*/
-
-/**
-*  Holds FullCrashDump properties from the config.
-*
-*	PDBCache_0_Branch=UE4-Branch
-*	PDBCache_0_ExecutablePathPattern=ue4.net\Builds\UE4-Branch\%ENGINE_VERSION%
-*	PDBCache_0_SymbolPathPattern=ue4.net\Builds\UE4-Branch\%ENGINE_VERSION%
-*	
-*	If PDBCache_0_SymbolPathPattern is missing, the value from PDBCache_0_ExecutablePathPattern will be used
-*/
-struct FPDBCacheConfigEntry
-{
-	/** Initialization constructor. */
-	FPDBCacheConfigEntry( const FString& InBranch, const FString& InExecutablePathPattern, const FString& InSymbolPathPattern )
-		: Branch( InBranch )
-		, ExecutablePathPattern( InExecutablePathPattern )
-		, SymbolPathPattern( InSymbolPathPattern )
-	{}
-
-
-	/** Branch name. */
-	const FString Branch;
-
-	/** Location of the executables */
-	const FString ExecutablePathPattern;
-
-	/** Location of the symbols, usually the same as the executables. */
-	const FString SymbolPathPattern;
-};
-
-/** Helper struct for reading PDB cache configuration. */
-struct FCrashDebugHelperConfig
-{
-	static FCrashDebugHelperConfig& Get()
-	{
-		static FCrashDebugHelperConfig Instance;
-		return Instance;
-	}
-
-	bool IsValid() const
-	{
-		// We need a least one entry to proceed.
-		return PDBCacheConfigEntries.Num() > 0;
-	}
-
-	/** Reads configuration. */
-	void ReadFullCrashDumpConfigurations();
-
-	/** Gets the config for branch. */
-	FPDBCacheConfigEntry GetCacheConfigEntryForBranch( const FString& Branch ) const
-	{
-		const FString BranchFixed = Branch.Replace( TEXT( "/" ), TEXT( "+" ) );
-		for (const auto& It : PDBCacheConfigEntries)
-		{
-			if (BranchFixed.Contains( It.Branch ))
-			{
-				return It;
-			}
-		}
-
-		// Invalid entry.
-		return FPDBCacheConfigEntry( Branch, FString(), FString() );
-	}
-
-protected:
-	/** Returns empty string if couldn't read. */
-	FString GetKey( const FString& KeyName );
-
-	/** Configuration for PDB Cache. */
-	TArray<FPDBCacheConfigEntry> PDBCacheConfigEntries;
-};
-
-void FCrashDebugHelperConfig::ReadFullCrashDumpConfigurations()
-{
-	for (int32 NumEntries = 0;; ++NumEntries)
-	{
-		const FString Branch = GetKey( FString::Printf( TEXT( "PDBCache_%i_Branch" ), NumEntries ) );
-		if (Branch.IsEmpty())
-		{
-			break;
-		}
-
-		const FString ExecutablePathPattern = GetKey( FString::Printf( TEXT( "PDBCache_%i_ExecutablePathPattern" ), NumEntries ) );
-		if (ExecutablePathPattern.IsEmpty())
-		{
-			break;
-		}
-
-		FString SymbolPathPattern = GetKey( FString::Printf( TEXT( "PDBCache_%i_SymbolPathPattern" ), NumEntries ) );
-		if (SymbolPathPattern.IsEmpty())
-		{
-			SymbolPathPattern = ExecutablePathPattern;
-		}
-		
-		PDBCacheConfigEntries.Add( FPDBCacheConfigEntry( Branch, ExecutablePathPattern, SymbolPathPattern ) );
-
-		UE_LOG( LogCrashDebugHelper, Log, TEXT( "PDBCacheConfigEntry: Branch:%s ExecutablePathPattern:%s SymbolPathPattern:%s" ), *Branch, *ExecutablePathPattern, *SymbolPathPattern );
-	}
-}
-
-FString FCrashDebugHelperConfig::GetKey( const FString& KeyName )
-{
-	const FString SectionName = TEXT( "Engine.CrashDebugHelper" );
-
-	FString Result;
-	if (!GConfig->GetString( *SectionName, *KeyName, Result, GEngineIni ))
-	{
-		return TEXT( "" );
-	}
-	return Result;
-}
 
 /*-----------------------------------------------------------------------------
 	ICrashDebugHelper
 -----------------------------------------------------------------------------*/
 
-void ICrashDebugHelper::SetDepotIndex( FString& PathToChange )
-{
-	FString CmdDepotIndex;
-	FParse::Value( FCommandLine::Get(), TEXT( "DepotIndex=" ), CmdDepotIndex );
-	// %DEPOT_INDEX% - Index of the depot, when multiple processor are used.
-	PathToChange.ReplaceInline( TEXT( "%DEPOT_INDEX%" ), *CmdDepotIndex );
-}
-
 bool ICrashDebugHelper::Init()
 {
 	bInitialized = true;
 
-	CrashInfo.bMutexPDBCache = FParse::Param(FCommandLine::Get(), TEXT("MutexPDBCache"));
-	FParse::Value(FCommandLine::Get(), TEXT("PDBCacheLock="), CrashInfo.PDBCacheLockName);
-
 	// Check if we have a valid EngineVersion, if so use it.
 	FString CmdEngineVersion;
 	const bool bHasEngineVersion = FParse::Value( FCommandLine::Get(), TEXT( "EngineVersion=" ), CmdEngineVersion );
@@ -225,391 +84,12 @@ bool ICrashDebugHelper::Init()
 	UE_LOG( LogCrashDebugHelper, Log, TEXT( "EngineVersion: %s" ), *CrashInfo.EngineVersion );
 	UE_LOG( LogCrashDebugHelper, Log, TEXT( "BuildVersion: %s" ), *CrashInfo.BuildVersion );
 
-	GConfig->GetString( TEXT( "Engine.CrashDebugHelper" ), TEXT( "SourceControlBuildLabelPattern" ), SourceControlBuildLabelPattern, GEngineIni );
-
-	FCrashDebugHelperConfig::Get().ReadFullCrashDumpConfigurations();
-
-	const bool bUsePDBCache = FCrashDebugHelperConfig::Get().IsValid();
-	UE_CLOG( !bUsePDBCache, LogCrashDebugHelper, Warning, TEXT( "CrashDebugHelperConfig invalid" ) );
-
-	if (bUsePDBCache)
-	{
-		if (CrashInfo.bMutexPDBCache && !CrashInfo.PDBCacheLockName.IsEmpty())
-		{
-			// Scoped lock
-			FSystemWideCriticalSection PDBCacheLock(CrashInfo.PDBCacheLockName, FTimespan::FromMinutes(10.0));
-			if (PDBCacheLock.IsValid())
-			{
-				FPDBCache::Get().Init();
-			}
-		}
-		else
-		{
-			FPDBCache::Get().Init();
-		}
-
-		if (!FPDBCache::Get().UsePDBCache())
-		{
-			UE_LOG(LogCrashDebugHelper, Warning, TEXT("PDB Cache failed to initialize"));
-		}
-	}
-	else
-	{
-		UE_LOG( LogCrashDebugHelper, Warning, TEXT( "PDB Cache disabled" ) );
-	}	
-
 	return bInitialized;
 }
 
-/** 
- * Initialise the source control interface, and ensure we have a valid connection
- */
-bool ICrashDebugHelper::InitSourceControl(bool bShowLogin)
-{
-	// Ensure we are in a valid state to sync
-	if (bInitialized == false)
-	{
-		UE_LOG(LogCrashDebugHelper, Warning, TEXT("InitSourceControl: CrashDebugHelper is not initialized properly."));
-		return false;
-	}
-
-#if UE_EDITOR || IS_PROGRAM
-	// Initialize the source control if it hasn't already been
-	if( !ISourceControlModule::Get().IsEnabled() || !ISourceControlModule::Get().GetProvider().IsAvailable() )
-	{
-		// make sure our provider is set to Perforce
-		ISourceControlModule::Get().SetProvider("Perforce");
-
-		// Attempt to load in a source control module
-		ISourceControlModule::Get().GetProvider().Init();
-
-#if !MINIDUMPDIAGNOSTICS
-		if ((ISourceControlModule::Get().GetProvider().IsAvailable() == false) || bShowLogin)
-		{
-			// Unable to connect? Prompt the user for login information
-			ISourceControlModule::Get().ShowLoginDialog(FSourceControlLoginClosed(), ELoginWindowMode::Modeless, EOnLoginWindowStartup::PreserveProvider);
-		}
-#endif
-		// If it's still disabled, none was found, so exit
-		if( !ISourceControlModule::Get().IsEnabled() || !ISourceControlModule::Get().GetProvider().IsAvailable() )
-		{
-			UE_LOG(LogCrashDebugHelper, Warning, TEXT("InitSourceControl: Source control unavailable or disabled."));
-			return false;
-		}
-	}
-#endif
-
-	return true;
-}
-
-/** 
- * Shutdown the connection to source control
- */
-void ICrashDebugHelper::ShutdownSourceControl()
-{
-#if UE_EDITOR || IS_PROGRAM
-	ISourceControlModule::Get().GetProvider().Close();
-#endif
-}
-
-
-bool ICrashDebugHelper::SyncModules(bool& bOutPDBCacheEntryValid)
-{
-	bOutPDBCacheEntryValid = false;
-
-	if( !FPDBCache::Get().UsePDBCache() )
-	{
-		UE_LOG( LogCrashDebugHelper, Warning, TEXT( "The PDB Cache is disabled, cannot proceed, %s" ), *CrashInfo.EngineVersion );
-		return false;
-	}
-
-	// #CrashReport: 2015-02-23 Obsolete, remove after 4.8
-	const TCHAR* UESymbols = TEXT( "Rocket/Symbols/" );
-	const bool bHasExecutable = !CrashInfo.ExecutablesPath.IsEmpty();
-	const bool bHasSymbols = !CrashInfo.SymbolsPath.IsEmpty();
-	
-	if( bHasExecutable && bHasSymbols )
-	{
-		if(FPDBCache::Get().ContainsPDBCacheEntry(CrashInfo.BuildVersion))
-		{
-			UE_LOG( LogCrashDebugHelper, Warning, TEXT( "Using cached storage: %s" ), *CrashInfo.BuildVersion);
-			CrashInfo.PDBCacheEntry = FPDBCache::Get().FindAndTouchPDBCacheEntry( CrashInfo.BuildVersion);
-		}
-		else
-		{
-			SCOPE_LOG_TIME_IN_SECONDS( TEXT( "SyncExecutableAndSymbolsFromNetwork" ), nullptr );
-
-			// Find all executables.
-			TArray<FString> NetworkExecutables;
-
-			// Don't duplicate work.
-			if (CrashInfo.ExecutablesPath != CrashInfo.SymbolsPath)
-			{
-				IFileManager::Get().FindFilesRecursive( NetworkExecutables, *CrashInfo.ExecutablesPath, TEXT( "*.dll" ), true, false, false );
-				IFileManager::Get().FindFilesRecursive( NetworkExecutables, *CrashInfo.ExecutablesPath, TEXT( "*.exe" ), true, false, false );
-			}
-
-			// Find all symbols.
-			TArray<FString> NetworkSymbols;
-			IFileManager::Get().FindFilesRecursive( NetworkSymbols, *CrashInfo.SymbolsPath, TEXT( "*.pdb" ), true, false, false );
-			IFileManager::Get().FindFilesRecursive( NetworkSymbols, *CrashInfo.SymbolsPath, TEXT( "*.dll" ), true, false, false );
-			IFileManager::Get().FindFilesRecursive( NetworkSymbols, *CrashInfo.SymbolsPath, TEXT( "*.exe" ), true, false, false );
-
-			// From=Full pathname
-			// To=Relative pathname
-			TMap<FString, FString> FilesToBeCached;
-
-			for( const auto& ExecutablePath : NetworkExecutables )
-			{
-				const FString NetworkRelativePath = ExecutablePath.Replace( *CrashInfo.ExecutablesPath, TEXT( "" ) );
-				FilesToBeCached.Add( ExecutablePath, NetworkRelativePath );
-			}
-
-			for( const auto& SymbolPath : NetworkSymbols )
-			{
-				const FString SymbolRelativePath = SymbolPath.Replace( *CrashInfo.SymbolsPath, TEXT( "" ) );
-				FilesToBeCached.Add( SymbolPath, SymbolRelativePath );
-			}
-
-			// Initialize and add a new PDB Cache entry to the database.
-			CrashInfo.PDBCacheEntry = FPDBCache::Get().CreateAndAddPDBCacheEntryMixed( CrashInfo.BuildVersion, FilesToBeCached );
-		}
-	}
-	// OBSOLETE PATH
-	else
-	{
-		// Command line for blocking obsolete path
-#if UE_EDITOR || IS_PROGRAM
-		const bool bNoP4Symbols = FParse::Param(FCommandLine::Get(), TEXT("NoP4Symbols"));
-
-		// Check source control
-		if (bNoP4Symbols || !ISourceControlModule::Get().IsEnabled())
-		{
-			return false;
-		}
-
-		// Get all labels associated with the crash info's label.
-		TArray< TSharedRef<ISourceControlLabel> > Labels = ISourceControlModule::Get().GetProvider().GetLabels(CrashInfo.LabelName);
-
-		if (Labels.Num() >= 1)
-		{
-			TSharedRef<ISourceControlLabel> Label = Labels[0];
-			TSet<FString> FilesToSync;
-
-			if (FPDBCache::Get().ContainsPDBCacheEntry(CrashInfo.EngineVersion))
-			{
-				UE_LOG(LogCrashDebugHelper, Warning, TEXT("Using cached storage: %s"), *CrashInfo.EngineVersion);
-				CrashInfo.PDBCacheEntry = FPDBCache::Get().FindAndTouchPDBCacheEntry(CrashInfo.EngineVersion);
-			}
-			// Use product version instead of label name to make a distinguish between chosen methods.
-			else if (FPDBCache::Get().ContainsPDBCacheEntry(CrashInfo.LabelName))
-			{
-				UE_LOG(LogCrashDebugHelper, Warning, TEXT("Using cached storage: %s"), *CrashInfo.LabelName);
-				CrashInfo.PDBCacheEntry = FPDBCache::Get().FindAndTouchPDBCacheEntry(CrashInfo.LabelName);
-			}
-			else if (bHasExecutable)
-			{
-				SCOPE_LOG_TIME_IN_SECONDS(TEXT("SyncModulesAndNetwork"), nullptr);
-
-				// Grab information about symbols.
-				TArray< TSharedRef<class ISourceControlRevision, ESPMode::ThreadSafe> > PDBSourceControlRevisions;
-				const FString PDBsPath = FString::Printf(TEXT("%s/%s....pdb"), *CrashInfo.DepotName, UESymbols);
-				Label->GetFileRevisions(PDBsPath, PDBSourceControlRevisions);
-
-				TSet<FString> PDBPaths;
-				for (const auto& PDBSrc : PDBSourceControlRevisions)
-				{
-					PDBPaths.Add(PDBSrc->GetFilename());
-				}
-
-				// Now, sync symbols.
-				for (const auto& PDBPath : PDBPaths)
-				{
-					if (Label->Sync(PDBPath))
-					{
-						UE_LOG(LogCrashDebugHelper, Warning, TEXT("Synced PDB: %s"), *PDBPath);
-					}
-				}
-
-				// Find all the executables in the product network path.
-				TArray<FString> NetworkExecutables;
-				IFileManager::Get().FindFilesRecursive(NetworkExecutables, *CrashInfo.ExecutablesPath, TEXT("*.dll"), true, false, false);
-				IFileManager::Get().FindFilesRecursive(NetworkExecutables, *CrashInfo.ExecutablesPath, TEXT("*.exe"), true, false, false);
-
-				// From=Full pathname
-				// To=Relative pathname
-				TMap<FString, FString> FilesToBeCached;
-
-				// If a symbol matches an executable, add the pair to the list of files that should be cached.
-				for (const auto& NetworkExecutableFullpath : NetworkExecutables)
-				{
-					for (const auto& PDBPath : PDBPaths)
-					{
-						const FString PDBRelativePath = PDBPath.Replace(*CrashInfo.DepotName, TEXT("")).Replace(UESymbols, TEXT(""));
-						const FString PDBFullpath = FPDBCache::Get().GetDepotRoot() / PDBPath;
-
-						const FString PDBMatch = PDBRelativePath.Replace(TEXT("pdb"), TEXT(""));
-						const FString NetworkRelativePath = NetworkExecutableFullpath.Replace(*CrashInfo.ExecutablesPath, TEXT(""));
-						const bool bMatch = NetworkExecutableFullpath.Contains(PDBMatch);
-						if (bMatch)
-						{
-							// From -> Where
-							FilesToBeCached.Add(NetworkExecutableFullpath, NetworkRelativePath);
-							FilesToBeCached.Add(PDBFullpath, PDBRelativePath);
-							break;
-						}
-					}
-				}
-
-				// Initialize and add a new PDB Cache entry to the database.
-				CrashInfo.PDBCacheEntry = FPDBCache::Get().CreateAndAddPDBCacheEntryMixed(CrashInfo.EngineVersion, FilesToBeCached);
-			}
-			else
-			{
-				TArray<FString> FilesToBeCached;
-
-				//@TODO: MAC: Excluding labels for Mac since we are only syncing windows binaries here...
-				if (Label->GetName().Contains(TEXT("Mac")))
-				{
-					UE_LOG(LogCrashDebugHelper, Log, TEXT("Skipping Mac label: %s"), *Label->GetName());
-				}
-				else
-				{
-					// Sync all the dll, exes, and related symbol files
-					UE_LOG(LogCrashDebugHelper, Log, TEXT("Syncing modules with label: %s"), *Label->GetName());
-
-					SCOPE_LOG_TIME_IN_SECONDS(TEXT("SyncModules"), nullptr);
-
-					// Grab all dll and pdb files for the specified label.
-					TArray< TSharedRef<class ISourceControlRevision, ESPMode::ThreadSafe> > DLLSourceControlRevisions;
-					const FString DLLsPath = FString::Printf(TEXT("%s/....dll"), *CrashInfo.DepotName);
-					Label->GetFileRevisions(DLLsPath, DLLSourceControlRevisions);
-
-					TArray< TSharedRef<class ISourceControlRevision, ESPMode::ThreadSafe> > EXESourceControlRevisions;
-					const FString EXEsPath = FString::Printf(TEXT("%s/....exe"), *CrashInfo.DepotName);
-					Label->GetFileRevisions(EXEsPath, EXESourceControlRevisions);
-
-					TArray< TSharedRef<class ISourceControlRevision, ESPMode::ThreadSafe> > PDBSourceControlRevisions;
-					const FString PDBsPath = FString::Printf(TEXT("%s/....pdb"), *CrashInfo.DepotName);
-					Label->GetFileRevisions(PDBsPath, PDBSourceControlRevisions);
-
-					TSet<FString> ModulesPaths;
-					for (const auto& DLLSrc : DLLSourceControlRevisions)
-					{
-						ModulesPaths.Add(DLLSrc->GetFilename().Replace(*CrashInfo.DepotName, TEXT("")));
-					}
-					for (const auto& EXESrc : EXESourceControlRevisions)
-					{
-						ModulesPaths.Add(EXESrc->GetFilename().Replace(*CrashInfo.DepotName, TEXT("")));
-					}
-
-					TSet<FString> PDBPaths;
-					for (const auto& PDBSrc : PDBSourceControlRevisions)
-					{
-						PDBPaths.Add(PDBSrc->GetFilename().Replace(*CrashInfo.DepotName, TEXT("")));
-					}
-
-					// Iterate through all module and see if we have dll and pdb associated with the module, if so add it to the files to sync.
-					for (const auto& ModuleName : CrashInfo.ModuleNames)
-					{
-						const FString ModuleNamePDB = ModuleName.Replace(TEXT(".dll"), TEXT(".pdb")).Replace(TEXT(".exe"), TEXT(".pdb"));
-
-						for (const auto& ModulePath : ModulesPaths)
-						{
-							const bool bContainsModule = ModulePath.Contains(ModuleName);
-							if (bContainsModule)
-							{
-								FilesToSync.Add(ModulePath);
-							}
-						}
-
-						for (const auto& PDBPath : PDBPaths)
-						{
-							const bool bContainsPDB = PDBPath.Contains(ModuleNamePDB);
-							if (bContainsPDB)
-							{
-								FilesToSync.Add(PDBPath);
-							}
-						}
-					}
-
-					// Now, sync all files.
-					for (const auto& Filename : FilesToSync)
-					{
-						const FString DepotPath = CrashInfo.DepotName + Filename;
-						if (Label->Sync(DepotPath))
-						{
-							UE_LOG(LogCrashDebugHelper, Warning, TEXT("Synced binary: %s"), *DepotPath);
-						}
-						FilesToBeCached.Add(DepotPath);
-					}
-				}
-
-				// Initialize and add a new PDB Cache entry to the database.
-				CrashInfo.PDBCacheEntry = FPDBCache::Get().CreateAndAddPDBCacheEntry(CrashInfo.LabelName, CrashInfo.DepotName, FilesToBeCached);
-			}
-		}
-		else
-		{
-			UE_LOG(LogCrashDebugHelper, Error, TEXT("Could not find label: %s"), *CrashInfo.LabelName);
-			return false;
-		}
-#endif
-	}
-
-	bOutPDBCacheEntryValid = CrashInfo.PDBCacheEntry.IsValid() && CrashInfo.PDBCacheEntry->Files.Num() > 0;
-	return true;
-}
-
-bool ICrashDebugHelper::SyncSourceFile()
-{
-#if UE_EDITOR || IS_PROGRAM
-	// Check source control
-	if( !ISourceControlModule::Get().IsEnabled() )
-	{
-		return false;
-	}
-
-	// Sync a single source file to requested CL.
-	FString DepotPath = CrashInfo.DepotName / CrashInfo.SourceFile + TEXT( "@" ) + TTypeToString<int32>::ToString( CrashInfo.BuiltFromCL );
-	ISourceControlModule::Get().GetProvider().Execute(ISourceControlOperation::Create<FSync>(), DepotPath);
-
-	UE_LOG( LogCrashDebugHelper, Warning, TEXT( "Syncing a single source file: %s"), *DepotPath );
-#endif
-
-	return true;
-}
-
-
 bool ICrashDebugHelper::ReadSourceFile( TArray<FString>& OutStrings )
 {
-	const bool bUsePDBCache = FPDBCache::Get().UsePDBCache();
-
-	FString FilePath;
-	if (bUsePDBCache)
-	{
-		// We assume a special folder for syncing all streams and //depot/ is not used in the view mapping.
-		/*
-			//depot/... //machine/... 
-			//UE4/...	//machine/Stream/UE4/...
-		*/
-		const TCHAR* DepotDelimiter = TEXT( "//depot/" );
-		const bool bIsDepot = CrashInfo.DepotName.Contains( DepotDelimiter );
-
-		if (bIsDepot)
-		{
-			FilePath = FPDBCache::Get().GetDepotRoot() / CrashInfo.DepotName.Replace( DepotDelimiter, TEXT( "" ) ) / CrashInfo.SourceFile;
-		}
-		else
-		{
-			FilePath = FPDBCache::Get().GetDepotRoot() / TEXT( "Stream" ) / CrashInfo.DepotName / CrashInfo.SourceFile;
-			FilePath.ReplaceInline( TEXT( "//" ), TEXT( "/" ) );
-		}
-	}
-	else
-	{
-		FilePath = FPaths::RootDir() / CrashInfo.SourceFile;
-	}
+	const FString FilePath = FPaths::RootDir() / CrashInfo.SourceFile;
 
 	FString Line;
 	if (FFileHelper::LoadFileToString( Line, *FilePath ))
@@ -650,53 +130,12 @@ void ICrashDebugHelper::AddSourceToReport()
 	}
 }
 
-bool ICrashDebugHelper::AddAnnotatedSourceToReport()
-{
-	// Make sure we have a source file to interrogate
-	if( CrashInfo.SourceFile.Len() > 0 && CrashInfo.SourceLineNumber != 0 && !CrashInfo.LabelName.IsEmpty() )
-	{
-#if UE_EDITOR || IS_PROGRAM
-		// Check source control
-		if( !ISourceControlModule::Get().IsEnabled() )
-		{
-			return false;
-		}
-
-		// Ask source control to annotate the file for us
-		FString DepotPath = CrashInfo.DepotName / CrashInfo.SourceFile;
-
-		TArray<FAnnotationLine> Lines;
-		SourceControlHelpers::AnnotateFile( ISourceControlModule::Get().GetProvider(), CrashInfo.BuiltFromCL, DepotPath, Lines );
-
-		uint32 MinLine = FMath::Clamp( CrashInfo.SourceLineNumber - 15, (uint32)1, (uint32)Lines.Num() );
-		uint32 MaxLine = FMath::Clamp( CrashInfo.SourceLineNumber + 15, (uint32)1, (uint32)Lines.Num() );
-
-		// Display a source context in the report, and decorate each line with the last editor of the line
-		for( uint32 Line = MinLine; Line < MaxLine; Line++ )
-		{			
-			if( Line == CrashInfo.SourceLineNumber )
-			{
-				CrashInfo.SourceContext.Add( FString::Printf( TEXT( "%5u ***** %20s: %s" ), Line, *Lines[Line].UserName, *Lines[Line].Line ) );
-			}
-			else
-			{
-				CrashInfo.SourceContext.Add( FString::Printf( TEXT( "%5u       %20s: %s" ), Line, *Lines[Line].UserName, *Lines[Line].Line ) );
-			}
-		}
-#endif
-		return true;
-	}
-
-	return false;
-}
-
 void FCrashInfo::Log( FString Line )
 {
 	UE_LOG( LogCrashDebugHelper, Warning, TEXT("%s"), *Line );
 	Report += Line + LINE_TERMINATOR;
 }
 
-
 const TCHAR* FCrashInfo::GetProcessorArchitecture( EProcessorArchitecture PA )
 {
 	switch( PA )
@@ -712,33 +151,23 @@ const TCHAR* FCrashInfo::GetProcessorArchitecture( EProcessorArchitecture PA )
 	return TEXT( "Unknown" );
 }
 
-
-int64 FCrashInfo::StringSize( const ANSICHAR* Line )
+void FCrashInfo::WriteLine( FArchive* ReportFile, const TCHAR* Line )
 {
-	int64 Size = 0;
-	if( Line != nullptr )
+	if (Line)
 	{
-		while( *Line++ != 0 )
-		{
-			Size++;
-		}
+		FTCHARToUTF8 UTF8Line(Line);
+		ReportFile->Serialize((void*)UTF8Line.Get(), UTF8Line.Length());
 	}
-	return Size;
+	ReportFile->Serialize(TCHAR_TO_UTF8(LINE_TERMINATOR), FCStringWide::Strlen(LINE_TERMINATOR));
 }
 
-
-void FCrashInfo::WriteLine( FArchive* ReportFile, const ANSICHAR* Line )
+void FCrashInfo::WriteLine(FArchive* ReportFile, const FString& Line)
 {
-	if( Line != NULL )
-	{
-		int64 StringBytes = StringSize( Line );
-		ReportFile->Serialize( ( void* )Line, StringBytes );
-	}
-
-	ReportFile->Serialize( TCHAR_TO_UTF8( LINE_TERMINATOR ), FCStringWide::Strlen(LINE_TERMINATOR) );
+	FTCHARToUTF8 UTF8Line(*Line, Line.Len());
+	ReportFile->Serialize((void*)UTF8Line.Get(), UTF8Line.Length());
+	ReportFile->Serialize(TCHAR_TO_UTF8(LINE_TERMINATOR), FCStringWide::Strlen(LINE_TERMINATOR));
 }
 
-
 void FCrashInfo::GenerateReport( const FString& DiagnosticsPath )
 {
 	FArchive* ReportFile = IFileManager::Get().CreateFileWriter( *DiagnosticsPath );
@@ -746,69 +175,69 @@ void FCrashInfo::GenerateReport( const FString& DiagnosticsPath )
 	{
 		FString Line;
 
-		WriteLine( ReportFile, TCHAR_TO_UTF8( TEXT( "Generating report for minidump" ) ) );
+		WriteLine( ReportFile, TEXT( "Generating report for minidump" ) );
 		WriteLine( ReportFile );
 
 		if ( EngineVersion.Len() > 0 )
 		{
 			Line = FString::Printf( TEXT( "Application version %s" ), *EngineVersion );
-			WriteLine( ReportFile, TCHAR_TO_UTF8( *Line ) );
+			WriteLine( ReportFile, Line );
 		}
 		else if( Modules.Num() > 0 )
 		{
 			Line = FString::Printf( TEXT( "Application version %d.%d.%d" ), Modules[0].Major, Modules[0].Minor, Modules[0].Patch );
-			WriteLine( ReportFile, TCHAR_TO_UTF8( *Line ) );
+			WriteLine( ReportFile, Line );
 		}
 
 		Line = FString::Printf( TEXT( " ... built from changelist %d" ), BuiltFromCL );
-		WriteLine( ReportFile, TCHAR_TO_UTF8( *Line ) );
+		WriteLine( ReportFile, Line );
 		if( LabelName.Len() > 0 )
 		{
 			Line = FString::Printf( TEXT( " ... based on label %s" ), *LabelName );
-			WriteLine( ReportFile, TCHAR_TO_UTF8( *Line ) );
+			WriteLine( ReportFile, Line );
 		}
 		WriteLine( ReportFile );
 
 		Line = FString::Printf( TEXT( "OS version %d.%d.%d.%d" ), SystemInfo.OSMajor, SystemInfo.OSMinor, SystemInfo.OSBuild, SystemInfo.OSRevision );
-		WriteLine( ReportFile, TCHAR_TO_UTF8( *Line ) );
+		WriteLine( ReportFile, Line );
 
 		Line = FString::Printf( TEXT( "Running %d %s processors" ), SystemInfo.ProcessorCount, GetProcessorArchitecture( SystemInfo.ProcessorArchitecture ) );
-		WriteLine( ReportFile, TCHAR_TO_UTF8( *Line ) );
+		WriteLine( ReportFile, Line );
 
 		Line = FString::Printf( TEXT( "Exception was \"%s\"" ), *Exception.ExceptionString );
-		WriteLine( ReportFile, TCHAR_TO_UTF8( *Line ) );
+		WriteLine( ReportFile, Line );
 		WriteLine( ReportFile );
 
 		Line = FString::Printf( TEXT( "Source context from \"%s\"" ), *SourceFile );
-		WriteLine( ReportFile, TCHAR_TO_UTF8( *Line ) );
+		WriteLine( ReportFile, Line );
 		WriteLine( ReportFile );
 
 		Line = FString::Printf( TEXT( "<SOURCE START>" ) );
-		WriteLine( ReportFile, TCHAR_TO_UTF8( *Line ) );
+		WriteLine( ReportFile, Line );
 		for( int32 LineIndex = 0; LineIndex < SourceContext.Num(); LineIndex++ )
 		{
 			Line = FString::Printf( TEXT( "%s" ), *SourceContext[LineIndex] );
-			WriteLine( ReportFile, TCHAR_TO_UTF8( *Line ) );
+			WriteLine( ReportFile, Line );
 		}
 		Line = FString::Printf( TEXT( "<SOURCE END>" ) );
-		WriteLine( ReportFile, TCHAR_TO_UTF8( *Line ) );
+		WriteLine( ReportFile, Line );
 		WriteLine( ReportFile );
 
 		Line = FString::Printf( TEXT( "<CALLSTACK START>" ) );
-		WriteLine( ReportFile, TCHAR_TO_UTF8( *Line ) );
+		WriteLine( ReportFile, Line );
 
 		for( int32 StackIndex = 0; StackIndex < Exception.CallStackString.Num(); StackIndex++ )
 		{
 			Line = FString::Printf( TEXT( "%s" ), *Exception.CallStackString[StackIndex] );
-			WriteLine( ReportFile, TCHAR_TO_UTF8( *Line ) );
+			WriteLine( ReportFile, Line );
 		}
 
 		Line = FString::Printf( TEXT( "<CALLSTACK END>" ) );
-		WriteLine( ReportFile, TCHAR_TO_UTF8( *Line ) );
+		WriteLine( ReportFile, Line );
 		WriteLine( ReportFile );
 
 		Line = FString::Printf( TEXT( "%d loaded modules" ), Modules.Num() );
-		WriteLine( ReportFile, TCHAR_TO_UTF8( *Line ) );
+		WriteLine( ReportFile, Line );
 
 		for( int32 ModuleIndex = 0; ModuleIndex < Modules.Num(); ModuleIndex++ )
 		{
@@ -823,128 +252,18 @@ void FCrashInfo::GenerateReport( const FString& DiagnosticsPath )
 			ModuleDetail += FString::Printf( TEXT( " 0x%016x 0x%08x" ), Module.BaseOfImage, Module.SizeOfImage );
 			ModuleDetail += FString::Printf( TEXT( " %s" ), *ModuleDirectory );
 
-			WriteLine( ReportFile, TCHAR_TO_UTF8( *ModuleDetail ) );
+			WriteLine( ReportFile, ModuleDetail );
 		}
 
 		WriteLine( ReportFile );
 
 		// Write out the processor debugging log
-		WriteLine( ReportFile, TCHAR_TO_UTF8( *Report ) );
+		WriteLine( ReportFile, Report );
 
 		Line = FString::Printf( TEXT( "Report end!" ) );
-		WriteLine( ReportFile, TCHAR_TO_UTF8( *Line )  );
+		WriteLine( ReportFile, Line  );
 
 		ReportFile->Close();
 		delete ReportFile;
 	}
 }
-
-void ICrashDebugHelper::FindSymbolsAndBinariesStorage()
-{
-	CrashInfo.ExecutablesPath.Empty();
-	CrashInfo.SymbolsPath.Empty();
-	CrashInfo.LabelName.Empty();
-
-	if( CrashInfo.BuiltFromCL == FCrashInfo::INVALID_CHANGELIST )
-	{
-		UE_LOG( LogCrashDebugHelper, Warning, TEXT( "Invalid parameters" ) );
-		return;
-	}
-
-	UE_LOG( LogCrashDebugHelper, Log, TEXT( "Engine version: %s" ), *CrashInfo.EngineVersion );
-
-	const FPDBCacheConfigEntry ConfigEntry = FCrashDebugHelperConfig::Get().GetCacheConfigEntryForBranch( CrashInfo.DepotName );
-	FString ExecutablePathPattern = ConfigEntry.ExecutablePathPattern;
-	FString SymbolPathPattern = ConfigEntry.SymbolPathPattern;
-
-	if (!ExecutablePathPattern.IsEmpty() || !SymbolPathPattern.IsEmpty())
-	{
-		UE_LOG( LogCrashDebugHelper, Log, TEXT( "Using branch: %s" ), *CrashInfo.DepotName );
-	}
-	else
-	{
-		UE_LOG( LogCrashDebugHelper, Log, TEXT( "Branch not found: %s" ), *CrashInfo.DepotName );
-		return;
-	}
-	
-	const FString StrENGINE_VERSION = CrashInfo.EngineVersion;
-	const FString StrPLATFORM_NAME = CrashInfo.PlatformName;
-	const FString StrPLATFORM_VARIANT = CrashInfo.PlatformVariantName;
-	const FString StrOLD_ENGINE_VERSION = FString::Printf( TEXT( "%s-CL-%i" ), *CrashInfo.DepotName.Replace( TEXT( "+" ), TEXT( "/" ) ), CrashInfo.BuiltFromCL )
-		.Replace( TEXT("/"), TEXT("+") );
-	const FString StrBUILD_VERSION = CrashInfo.BuildVersion;
-
-	const FString TestExecutablesPath = ExecutablePathPattern
-		.Replace( TEXT( "%ENGINE_VERSION%" ), *StrENGINE_VERSION )
-		.Replace( TEXT( "%PLATFORM_NAME%" ), *StrPLATFORM_NAME )
-		.Replace( TEXT( "%PLATFORM_VARIANT%" ), *StrPLATFORM_VARIANT )
-		.Replace( TEXT( "%OLD_ENGINE_VERSION%" ), *StrOLD_ENGINE_VERSION )
-		.Replace( TEXT( "%BUILD_VERSION%" ), *StrBUILD_VERSION );
-
-	const FString TestSymbolsPath = SymbolPathPattern
-		.Replace( TEXT( "%ENGINE_VERSION%" ), *StrENGINE_VERSION )
-		.Replace( TEXT( "%PLATFORM_NAME%" ), *StrPLATFORM_NAME )
-		.Replace( TEXT( "%PLATFORM_VARIANT%" ), *StrPLATFORM_VARIANT )
-		.Replace( TEXT( "%OLD_ENGINE_VERSION%" ), *StrOLD_ENGINE_VERSION )
-		.Replace( TEXT( "%BUILD_VERSION%" ), *StrBUILD_VERSION );
-
-	// Try to find the network path by using the pattern supplied via ini.
-	// If this step successes, we will grab the executable from the network path instead of P4.
-	bool bFoundDirectory = false;
-	
-	bool bHasExecutables = false;
-	if (!TestExecutablesPath.IsEmpty())
-	{
-		bHasExecutables = IFileManager::Get().DirectoryExists(*TestExecutablesPath);
-	}
-	bool bHasSymbols = false;
-	if (!TestSymbolsPath.IsEmpty())
-	{
-		bHasSymbols = IFileManager::Get().DirectoryExists(*TestSymbolsPath);
-	}
-
-	if( bHasExecutables && bHasSymbols )
-	{
-		CrashInfo.ExecutablesPath = TestExecutablesPath;
-		CrashInfo.SymbolsPath = TestSymbolsPath;
-		bFoundDirectory = true;
-		UE_LOG( LogCrashDebugHelper, Log, TEXT( "Using path for executables and symbols: %s" ), *CrashInfo.ExecutablesPath );
-	}
-	else if( bHasExecutables )
-	{
-		CrashInfo.ExecutablesPath = TestExecutablesPath;
-		UE_LOG( LogCrashDebugHelper, Log, TEXT( "Using path for executables: %s" ), *CrashInfo.ExecutablesPath );
-	}
-	else
-	{
-		UE_LOG( LogCrashDebugHelper, Log, TEXT( "Path for executables not found: %s" ), *TestExecutablesPath );
-	}
-	
-
-	// Try to find the label directly in source control by using the pattern supplied via ini.
-	if( !bFoundDirectory && !SourceControlBuildLabelPattern.IsEmpty() )
-	{	
-		const FString ChangelistString = FString::Printf( TEXT( "%d" ), CrashInfo.BuiltFromCL );
-		const FString LabelWithCL = SourceControlBuildLabelPattern.Replace( TEXT( "%CHANGELISTNUMBER%" ), *ChangelistString, ESearchCase::CaseSensitive );
-		UE_LOG( LogCrashDebugHelper, Log, TEXT( "Label matching pattern: %s" ), *LabelWithCL );
-
-#if UE_EDITOR || IS_PROGRAM
-		TArray< TSharedRef<ISourceControlLabel> > Labels = ISourceControlModule::Get().GetProvider().GetLabels( LabelWithCL );
-		if( Labels.Num() > 0 )
-		{
-			const int32 LabelIndex = 0;
-			CrashInfo.LabelName = Labels[LabelIndex]->GetName();;
-		
-			// If we found more than one label, warn about it and just use the first one
-			if( Labels.Num() > 1 )
-			{
-				UE_LOG( LogCrashDebugHelper, Warning, TEXT( "More than one build label found, using label: %s" ), *LabelWithCL, *CrashInfo.LabelName );
-			}
-			else
-			{
-				UE_LOG( LogCrashDebugHelper, Log, TEXT( "Using label: %s" ), *CrashInfo.LabelName );
-			}		
-		}
-#endif
-	}
-}
diff --git a/Engine/Source/Developer/CrashDebugHelper/Private/CrashDebugPDBCache.cpp b/Engine/Source/Developer/CrashDebugHelper/Private/CrashDebugPDBCache.cpp
deleted file mode 100644
diff --git a/Engine/Source/Developer/CrashDebugHelper/Private/IOS/CrashDebugHelperIOS.cpp b/Engine/Source/Developer/CrashDebugHelper/Private/IOS/CrashDebugHelperIOS.cpp
--- a/Engine/Source/Developer/CrashDebugHelper/Private/IOS/CrashDebugHelperIOS.cpp
+++ b/Engine/Source/Developer/CrashDebugHelper/Private/IOS/CrashDebugHelperIOS.cpp
@@ -599,14 +599,6 @@ bool FCrashDebugHelperIOS::ParseCrashDump(const FString& InCrashDumpName, FCrash
 bool FCrashDebugHelperIOS::CreateMinidumpDiagnosticReport( const FString& InCrashDumpName )
 {
 	bool bOK = false;
-	const bool bSyncSymbols = FParse::Param( FCommandLine::Get(), TEXT( "SyncSymbols" ) );
-	const bool bAnnotate = FParse::Param( FCommandLine::Get(), TEXT( "Annotate" ) );
-	const bool bUseSCC = bSyncSymbols || bAnnotate;
-	
-	if( bUseSCC )
-	{
-		InitSourceControl( false );
-	}
 	
 	FString CrashDump;
 	if ( FFileHelper::LoadFileToString( CrashDump, *InCrashDumpName ) )
@@ -645,14 +637,6 @@ bool FCrashDebugHelperIOS::CreateMinidumpDiagnosticReport( const FString& InCras
 			if(Result == 5 && Branch.Len() > 0)
 			{
 				CrashInfo.LabelName = Branch;
-				
-				if( bSyncSymbols )
-				{
-					FindSymbolsAndBinariesStorage();
-					
-					bool bPDBCacheEntryValid = false;
-					SyncModules(bPDBCacheEntryValid);
-				}
 			}
 			
 			Result = ParseOS(*CrashDump, CrashInfo.SystemInfo.OSMajor, CrashInfo.SystemInfo.OSMinor, CrashInfo.SystemInfo.OSBuild, CrashInfo.SystemInfo.OSRevision);
@@ -767,24 +751,8 @@ bool FCrashDebugHelperIOS::CreateMinidumpDiagnosticReport( const FString& InCras
 								CrashInfo.SourceFile = ExtractRelativePath( TEXT( "source" ), *FileName );
 								CrashInfo.SourceLineNumber = LineNumber;
 								
-								if( bSyncSymbols && CrashInfo.BuiltFromCL > 0 )
-								{
-									UE_LOG( LogCrashDebugHelper, Log, TEXT( "Using CL %i to sync crash source file" ), CrashInfo.BuiltFromCL );
-									SyncSourceFile();
-								}
-								
-								// Try to annotate the file if requested
-								bool bAnnotationSuccessful = false;
-								if( bAnnotate )
-								{
-									bAnnotationSuccessful = AddAnnotatedSourceToReport();
-								}
-								
-								// If annotation is not requested, or failed, add the standard source context
-								if( !bAnnotationSuccessful )
-								{
-									AddSourceToReport();
-								}
+								// Add the standard source context
+								AddSourceToReport();
 							}
 						}
 						
@@ -811,10 +779,5 @@ bool FCrashDebugHelperIOS::CreateMinidumpDiagnosticReport( const FString& InCras
 		}
 	}
 	
-	if( bUseSCC )
-	{
-		ShutdownSourceControl();
-	}
-	
 	return bOK;
 }
\ No newline at end of file
diff --git a/Engine/Source/Developer/CrashDebugHelper/Private/Mac/CrashDebugHelperMac.cpp b/Engine/Source/Developer/CrashDebugHelper/Private/Mac/CrashDebugHelperMac.cpp
--- a/Engine/Source/Developer/CrashDebugHelper/Private/Mac/CrashDebugHelperMac.cpp
+++ b/Engine/Source/Developer/CrashDebugHelper/Private/Mac/CrashDebugHelperMac.cpp
@@ -606,14 +606,6 @@ bool FCrashDebugHelperMac::ParseCrashDump(const FString& InCrashDumpName, FCrash
 bool FCrashDebugHelperMac::CreateMinidumpDiagnosticReport( const FString& InCrashDumpName )
 {
 	bool bOK = false;
-	const bool bSyncSymbols = FParse::Param( FCommandLine::Get(), TEXT( "SyncSymbols" ) );
-	const bool bAnnotate = FParse::Param( FCommandLine::Get(), TEXT( "Annotate" ) );
-	const bool bUseSCC = bSyncSymbols || bAnnotate;
-	
-	if( bUseSCC )
-	{
-		InitSourceControl( false );
-	}
 	
 	FString CrashDump;
 	
@@ -671,14 +663,6 @@ bool FCrashDebugHelperMac::CreateMinidumpDiagnosticReport( const FString& InCras
 			if(Result == 5 && Branch.Len() > 0)
 			{
 				CrashInfo.LabelName = Branch;
-				
-				if( bSyncSymbols )
-				{
-					FindSymbolsAndBinariesStorage();
-					
-					bool bPDBCacheEntryValid = false;
-					SyncModules(bPDBCacheEntryValid);
-				}
 			}
 			
 			Result = ParseOS(*CrashDump, CrashInfo.SystemInfo.OSMajor, CrashInfo.SystemInfo.OSMinor, CrashInfo.SystemInfo.OSBuild, CrashInfo.SystemInfo.OSRevision);
@@ -793,24 +777,8 @@ bool FCrashDebugHelperMac::CreateMinidumpDiagnosticReport( const FString& InCras
 								CrashInfo.SourceFile = ExtractRelativePath( TEXT( "source" ), *FileName );
 								CrashInfo.SourceLineNumber = LineNumber;
 								
-								if( bSyncSymbols && CrashInfo.BuiltFromCL > 0 )
-								{
-									UE_LOG( LogCrashDebugHelper, Log, TEXT( "Using CL %i to sync crash source file" ), CrashInfo.BuiltFromCL );
-									SyncSourceFile();
-								}
-								
-								// Try to annotate the file if requested
-								bool bAnnotationSuccessful = false;
-								if( bAnnotate )
-								{
-									bAnnotationSuccessful = AddAnnotatedSourceToReport();
-								}
-								
-								// If annotation is not requested, or failed, add the standard source context
-								if( !bAnnotationSuccessful )
-								{
-									AddSourceToReport();
-								}
+								// Add the standard source context
+								AddSourceToReport();
 							}
 						}
 						
@@ -837,10 +805,5 @@ bool FCrashDebugHelperMac::CreateMinidumpDiagnosticReport( const FString& InCras
 		}
 	}
 	
-	if( bUseSCC )
-	{
-		ShutdownSourceControl();
-	}
-	
 	return bOK;
 }
diff --git a/Engine/Source/Developer/CrashDebugHelper/Private/Windows/CrashDebugHelperWindows.cpp b/Engine/Source/Developer/CrashDebugHelper/Private/Windows/CrashDebugHelperWindows.cpp
--- a/Engine/Source/Developer/CrashDebugHelper/Private/Windows/CrashDebugHelperWindows.cpp
+++ b/Engine/Source/Developer/CrashDebugHelper/Private/Windows/CrashDebugHelperWindows.cpp
@@ -7,7 +7,6 @@
 #include "Misc/CommandLine.h"
 
 #include "Misc/EngineVersion.h"
-#include "ISourceControlModule.h"
 
 #include "Windows/WindowsHWrapper.h"
 #include "Windows/AllowWindowsPlatformTypes.h"
@@ -15,16 +14,6 @@
 
 bool FCrashDebugHelperWindows::CreateMinidumpDiagnosticReport( const FString& InCrashDumpFilename )
 {
-	const bool bSyncSymbols = FParse::Param( FCommandLine::Get(), TEXT( "SyncSymbols" ) );
-	const bool bAnnotate = FParse::Param( FCommandLine::Get(), TEXT( "Annotate" ) );
-	const bool bNoTrim = FParse::Param(FCommandLine::Get(), TEXT("NoTrimCallstack"));
-	const bool bUseSCC = bSyncSymbols || bAnnotate;
-
-	if( bUseSCC )
-	{
-		InitSourceControl( false );
-	}
-
 	FWindowsPlatformStackWalkExt WindowsStackWalkExt( CrashInfo );
 
 	const bool bReady = WindowsStackWalkExt.InitStackWalking();
@@ -39,147 +28,43 @@ bool FCrashDebugHelperWindows::CreateMinidumpDiagnosticReport( const FString& In
 			WindowsStackWalkExt.GetExeFileVersionAndModuleList(ExeFileVersion);
 
 			// Init Symbols
-			bool bInitSymbols = false;
-			if (CrashInfo.bMutexPDBCache && !CrashInfo.PDBCacheLockName.IsEmpty())
-			{
-				// Scoped lock
-				UE_LOG(LogCrashDebugHelper, Log, TEXT("Locking for InitSymbols()"));
-				FSystemWideCriticalSection PDBCacheLock(CrashInfo.PDBCacheLockName, FTimespan::FromMinutes(10.0));
-				if (PDBCacheLock.IsValid())
-				{
-					bInitSymbols = InitSymbols(WindowsStackWalkExt, bSyncSymbols);
-				}
-				UE_LOG(LogCrashDebugHelper, Log, TEXT("Unlocking after InitSymbols()"));
-			}
-			else
-			{
-				bInitSymbols = InitSymbols(WindowsStackWalkExt, bSyncSymbols);
-			}
-
-			if (bInitSymbols)
-			{
-				// Get all the info we should ever need about the modules
-				WindowsStackWalkExt.GetModuleInfoDetailed();
-
-				// Get info about the system that created the minidump
-				WindowsStackWalkExt.GetSystemInfo();
-
-				// Get all the thread info
-				WindowsStackWalkExt.GetThreadInfo();
-
-				// Get exception info
-				WindowsStackWalkExt.GetExceptionInfo();
-
-				// Get the callstacks for each thread
-				WindowsStackWalkExt.GetCallstacks();
-
-				// Sync the source file where the crash occurred
-				if (CrashInfo.SourceFile.Len() > 0)
-				{
-					const bool bMutexSourceSync = FParse::Param(FCommandLine::Get(), TEXT("MutexSourceSync"));
-					FString SourceSyncLockName;
-					FParse::Value(FCommandLine::Get(), TEXT("SourceSyncLock="), SourceSyncLockName);
-
-					if (bMutexSourceSync && !SourceSyncLockName.IsEmpty())
-					{
-						// Scoped lock
-						UE_LOG(LogCrashDebugHelper, Log, TEXT("Locking for SyncAndReadSourceFile()"));
-						const FTimespan GlobalLockWaitTimeout = FTimespan::FromSeconds(30.0);
-						FSystemWideCriticalSection SyncSourceLock(SourceSyncLockName, GlobalLockWaitTimeout);
-						if (SyncSourceLock.IsValid())
-						{
-							SyncAndReadSourceFile(bSyncSymbols, bAnnotate, CrashInfo.BuiltFromCL);
-						}
-						UE_LOG(LogCrashDebugHelper, Log, TEXT("Unlocking after SyncAndReadSourceFile()"));
-					}
-					else
-					{
-						SyncAndReadSourceFile(bSyncSymbols, bAnnotate, CrashInfo.BuiltFromCL);
-					}
-				}
-
-				// Set the result
-				bResult = true;
-			}
-			else
-			{
-				UE_LOG(LogCrashDebugHelper, Warning, TEXT("InitSymbols failed"));
-			}
-		}
-		else
-		{
-			UE_LOG( LogCrashDebugHelper, Warning, TEXT( "Invalid built from changelist" ) );
-		}
-	}
-	else
-	{
-		UE_LOG( LogCrashDebugHelper, Warning, TEXT( "Failed to open crash dump file: %s" ), *InCrashDumpFilename );
-	}
-
-	if( bUseSCC )
-	{
-		ShutdownSourceControl();
-	}
+			WindowsStackWalkExt.InitSymbols();
 
-	return bResult;
-}
+			// Set the symbol path based on the loaded modules
+			WindowsStackWalkExt.SetSymbolPathsFromModules();
 
-bool FCrashDebugHelperWindows::InitSymbols(FWindowsPlatformStackWalkExt& WindowsStackWalkExt, bool bSyncSymbols)
-{
-	// CrashInfo now contains a changelist to lookup a label for
-	if (bSyncSymbols)
-	{
-		FindSymbolsAndBinariesStorage();
+			// Get all the info we should ever need about the modules
+			WindowsStackWalkExt.GetModuleInfoDetailed();
 
-		bool bPDBCacheEntryValid = false;
-		const bool bSynced = SyncModules(bPDBCacheEntryValid);
-		// Without symbols we can't decode the provided minidump.
-		if (!bSynced)
-		{
-			return false;
-		}
+			// Get info about the system that created the minidump
+			WindowsStackWalkExt.GetSystemInfo();
 
-		if (!bPDBCacheEntryValid)
-		{
-			// early-out option
-			const bool bForceUsePDBCache = FParse::Param(FCommandLine::Get(), TEXT("ForceUsePDBCache"));
-			if (bForceUsePDBCache)
-			{
-				UE_LOG(LogCrashDebugHelper, Log, TEXT("No cached symbols available. Exiting due to -ForceUsePDBCache."));
-				return false;
-			}
-		}
-	}
+			// Get all the thread info
+			WindowsStackWalkExt.GetThreadInfo();
 
-	// Initialise the symbol options
-	WindowsStackWalkExt.InitSymbols();
+			// Get exception info
+			WindowsStackWalkExt.GetExceptionInfo();
 
-	// Set the symbol path based on the loaded modules
-	WindowsStackWalkExt.SetSymbolPathsFromModules();
+			// Get the callstacks for each thread
+			WindowsStackWalkExt.GetCallstacks();
 
-	return true;
-}
+			// Add the source file where the crash occurred
+			AddSourceToReport();
 
-void FCrashDebugHelperWindows::SyncAndReadSourceFile(bool bSyncSymbols, bool bAnnotate, int32 BuiltFromCL)
-{
-	if (bSyncSymbols && BuiltFromCL > 0)
-	{
-		UE_LOG(LogCrashDebugHelper, Log, TEXT("Using CL %i to sync crash source file"), BuiltFromCL);
-		SyncSourceFile();
+			// Set the result
+			bResult = true;
+		}
+		else
+		{
+			UE_LOG( LogCrashDebugHelper, Warning, TEXT( "Invalid built from changelist" ) );
+		}
 	}
-
-	// Try to annotate the file if requested
-	bool bAnnotationSuccessful = false;
-	if (bAnnotate)
+	else
 	{
-		bAnnotationSuccessful = AddAnnotatedSourceToReport();
+		UE_LOG( LogCrashDebugHelper, Warning, TEXT( "Failed to open crash dump file: %s" ), *InCrashDumpFilename );
 	}
 
-	// If annotation is not requested, or failed, add the standard source context
-	if (!bAnnotationSuccessful)
-	{
-		AddSourceToReport();
-	}
+	return bResult;
 }
 
 #include "Windows/HideWindowsPlatformTypes.h"
diff --git a/Engine/Source/Developer/CrashDebugHelper/Private/Windows/CrashDebugHelperWindows.h b/Engine/Source/Developer/CrashDebugHelper/Private/Windows/CrashDebugHelperWindows.h
--- a/Engine/Source/Developer/CrashDebugHelper/Private/Windows/CrashDebugHelperWindows.h
+++ b/Engine/Source/Developer/CrashDebugHelper/Private/Windows/CrashDebugHelperWindows.h
@@ -15,10 +15,6 @@ public:
 	 *	@return	bool				true if successful, false if not
 	 */
 	virtual bool CreateMinidumpDiagnosticReport( const FString& InCrashDumpName ) override;
-
-private:
-	bool InitSymbols(struct FWindowsPlatformStackWalkExt& WindowsStackWalkExt, bool bSyncSymbols);
-	void SyncAndReadSourceFile(bool bSyncSymbols, bool bAnnotate, int32 BuiltFromCL);
 };
 
 typedef FCrashDebugHelperWindows FCrashDebugHelper;
diff --git a/Engine/Source/Developer/CrashDebugHelper/Private/Windows/WindowsPlatformStackWalkExt.cpp b/Engine/Source/Developer/CrashDebugHelper/Private/Windows/WindowsPlatformStackWalkExt.cpp
--- a/Engine/Source/Developer/CrashDebugHelper/Private/Windows/WindowsPlatformStackWalkExt.cpp
+++ b/Engine/Source/Developer/CrashDebugHelper/Private/Windows/WindowsPlatformStackWalkExt.cpp
@@ -1,10 +1,10 @@
 // Copyright 1998-2019 Epic Games, Inc. All Rights Reserved.
 
 #include "WindowsPlatformStackWalkExt.h"
+#include "CrashDebugHelper.h"
 #include "CrashDebugHelperPrivate.h"
 #include "GenericPlatform/GenericPlatformStackWalk.h"
 #include "GenericPlatform/GenericPlatformCrashContext.h"
-#include "CrashDebugPDBCache.h"
 #include "Misc/Parse.h"
 #include "Misc/CommandLine.h"
 #include "Misc/MemStack.h"
@@ -179,59 +179,33 @@ void FWindowsPlatformStackWalkExt::GetExeFileVersionAndModuleList( FCrashModuleI
 
 void FWindowsPlatformStackWalkExt::SetSymbolPathsFromModules()
 {
-	const bool bUseCachedData = CrashInfo.PDBCacheEntry.IsValid();
 	FString CombinedPath = TEXT( "" );
 
-	// Use symbol cache from command line
-	FString DebugSymbols;
-	if (FParse::Value(FCommandLine::Get(), TEXT("DebugSymbols="), DebugSymbols))
 	{
-		CombinedPath += TEXT("SRV*");
-		CombinedPath += DebugSymbols;
-		CombinedPath += TEXT(";");
-	}
-
-	// For externally launched minidump diagnostics.
-	if( bUseCachedData )
-	{
-		TSet<FString> SymbolPaths;
-		for( const auto& Filename : CrashInfo.PDBCacheEntry->Files )
+		// Use symbol cache from command line
+		FString DebugSymbols;
+		if (FParse::Value(FCommandLine::Get(), TEXT("DebugSymbols="), DebugSymbols))
 		{
-			const FString SymbolPath = FPaths::GetPath( Filename );
-			if( SymbolPath.Len() > 0 )
-			{
-				SymbolPaths.Add( SymbolPath );
-			}
+			CombinedPath += TEXT("SRV*");
+			CombinedPath += DebugSymbols;
+			CombinedPath += TEXT(";");
 		}
 
-		for( const auto& SymbolPath : SymbolPaths )
-		{
-			CombinedPath += SymbolPath;
-			CombinedPath += TEXT( ";" );
-		}
-
-		// Set the symbol path
-		Symbol->SetImagePathWide( *CombinedPath );
-		Symbol->SetSymbolPathWide( *CombinedPath );
-	}
-	// For locally launched minidump diagnostics.
-	else
-	{
 		// Iterate over all loaded modules.
 		TSet<FString> SymbolPaths;
 		for (const auto& Filename : CrashInfo.ModuleNames)
 		{
-			const FString Path = FPaths::GetPath( Filename );
-			if( Path.Len() > 0 )
+			const FString Path = FPaths::GetPath(Filename);
+			if (Path.Len() > 0)
 			{
-				SymbolPaths.Add( Path );
+				SymbolPaths.Add(Path);
 			}
 		}
 
-		for( const auto& SymbolPath : SymbolPaths )
+		for (const auto& SymbolPath : SymbolPaths)
 		{
 			CombinedPath += SymbolPath;
-			CombinedPath += TEXT( ";" );
+			CombinedPath += TEXT(";");
 		}
 
 #if ALLOW_UNREAL_ACCESS_TO_NT_SYMBOL_PATH
@@ -242,12 +216,12 @@ void FWindowsPlatformStackWalkExt::SetSymbolPathsFromModules()
 			CombinedPath += ";";
 		}
 #endif
-
-		// Set the symbol path
-		Symbol->SetImagePathWide( *CombinedPath );
-		Symbol->SetSymbolPathWide( *CombinedPath );
 	}
 
+	// Set the symbol path
+	Symbol->SetImagePathWide( *CombinedPath );
+	Symbol->SetSymbolPathWide( *CombinedPath );
+
 	// Add in syncing of the Microsoft symbol servers if requested
 	if( FParse::Param( FCommandLine::Get(), TEXT( "SyncMicrosoftSymbols" ) ) )
 	{
diff --git a/Engine/Source/Developer/CrashDebugHelper/Public/CrashDebugHelper.h b/Engine/Source/Developer/CrashDebugHelper/Public/CrashDebugHelper.h
--- a/Engine/Source/Developer/CrashDebugHelper/Public/CrashDebugHelper.h
+++ b/Engine/Source/Developer/CrashDebugHelper/Public/CrashDebugHelper.h
@@ -7,7 +7,6 @@
 #include "Templates/SharedPointer.h"
 
 class FArchive;
-struct FPDBCacheEntry;
 
 enum EProcessorArchitecture
 {
@@ -27,22 +26,12 @@ public:
 
 	FString Name;
 	FString Extension;
-	uint64 BaseOfImage;
-	uint32 SizeOfImage;
-	uint16 Major;
-	uint16 Minor;
-	uint16 Patch;
-	uint16 Revision;
-
-	FCrashModuleInfo()
-		: BaseOfImage( 0 )
-		, SizeOfImage( 0 )
-		, Major( 0 )
-		, Minor( 0 )
-		, Patch( 0 )
-		, Revision( 0 )
-	{
-	}
+	uint64 BaseOfImage = 0;
+	uint32 SizeOfImage = 0;
+	uint16 Major = 0;
+	uint16 Minor = 0;
+	uint16 Patch = 0;
+	uint16 Revision = 0;
 };
 
 /** 
@@ -53,20 +42,10 @@ class FCrashThreadInfo
 public:
 	FString Report;
 
-	uint32 ThreadId;
-	uint32 SuspendCount;
+	uint32 ThreadId = 0;
+	uint32 SuspendCount = 0;
 
 	TArray<uint64> CallStack;
-
-	FCrashThreadInfo()
-		: ThreadId( 0 )
-		, SuspendCount( 0 )
-	{
-	}
-
-	~FCrashThreadInfo()
-	{
-	}
 };
 
 /** 
@@ -77,23 +56,12 @@ class FCrashExceptionInfo
 public:
 	FString Report;
 
-	uint32 ProcessId;
-	uint32 ThreadId;
-	uint32 Code;
+	uint32 ProcessId = 0;
+	uint32 ThreadId = 0;
+	uint32 Code = 0;
 	FString ExceptionString;
 
 	TArray<FString> CallStackString;
-
-	FCrashExceptionInfo()
-		: ProcessId( 0 )
-		, ThreadId( 0 )
-		, Code( 0 )
-	{
-	}
-
-	~FCrashExceptionInfo()
-	{
-	}
 };
 
 /** 
@@ -104,23 +72,13 @@ class FCrashSystemInfo
 public:
 	FString Report;
 
-	EProcessorArchitecture ProcessorArchitecture;
-	uint32 ProcessorCount;
-
-	uint16 OSMajor;
-	uint16 OSMinor;
-	uint16 OSBuild;
-	uint16 OSRevision;
-
-	FCrashSystemInfo()
-		: ProcessorArchitecture( PA_UNKNOWN )
-		, ProcessorCount( 0 )
-		, OSMajor( 0 )
-		, OSMinor( 0 )
-		, OSBuild( 0 )
-		, OSRevision( 0 )
-	{
-	}
+	EProcessorArchitecture ProcessorArchitecture = PA_UNKNOWN;
+	uint32 ProcessorCount = 0;
+
+	uint16 OSMajor = 0;
+	uint16 OSMinor = 0;
+	uint16 OSBuild = 0;
+	uint16 OSRevision = 0;
 };
 
 // #TODO 2015-07-24 Refactor
@@ -147,7 +105,7 @@ public:
 	FString BuildVersion;
 
 	/** CL built from. */
-	int32 BuiltFromCL;
+	int32 BuiltFromCL = INVALID_CHANGELIST;
 
 	/** The label the describes the executables and symbols. */
 	FString LabelName;
@@ -159,7 +117,7 @@ public:
 	FString SymbolsPath;
 
 	FString SourceFile;
-	uint32 SourceLineNumber;
+	uint32 SourceLineNumber = 0;
 	TArray<FString> SourceContext;
 
 	/** Only modules names, retrieved from the minidump file. */
@@ -170,28 +128,9 @@ public:
 	TArray<FCrashThreadInfo> Threads;
 	TArray<FCrashModuleInfo> Modules;
 
-	/** Shared pointer to the PDB Cache entry, if valid contains all information about synced PDBs. */
-	TSharedPtr<FPDBCacheEntry> PDBCacheEntry;
-
 	FString PlatformName;
 	FString PlatformVariantName;
 
-	/** If we are using a PDBCache, this is whether we should use a system-wide lock to access it. */
-	bool bMutexPDBCache;
-
-	/** If we are using a PDBCache, this is the name of the system-wide lock we should use to access it. */
-	FString PDBCacheLockName;
-
-	FCrashInfo()
-		: BuiltFromCL( INVALID_CHANGELIST )
-		, SourceLineNumber( 0 )
-	{
-	}
-
-	~FCrashInfo()
-	{
-	}
-
 	/** 
 	 * Generate a report for the crash in the requested path
 	 */
@@ -206,20 +145,15 @@ private:
 	/** 
 	 * Convert the processor architecture to a human readable string
 	 */
-	const TCHAR* GetProcessorArchitecture( EProcessorArchitecture PA );
-
-	/**
-	 * Calculate the byte size of a UTF-8 string
-	 */
-	int64 StringSize( const ANSICHAR* Line );
+	static const TCHAR* GetProcessorArchitecture( EProcessorArchitecture PA );
 
 	/** 
-	* Write a line of UTF-8 to a file
+	* Write a line as UTF-8 to a file
 	*/
-	void WriteLine( FArchive* ReportFile, const ANSICHAR* Line = NULL );
+	static void WriteLine(FArchive* ReportFile, const TCHAR* Line = nullptr);
+	static void WriteLine(FArchive* ReportFile, const FString& Line);
 };
 
-
 /** Helper structure for tracking crash debug information */
 struct FCrashDebugInfo
 {
@@ -237,28 +171,8 @@ struct FCrashDebugInfo
 class CRASHDEBUGHELPER_API ICrashDebugHelper
 {
 public:
-	/** Replaces %DEPOT_INDEX% with the command line DepotIndex in the specified path. */
-	static void SetDepotIndex( FString& PathToChange );
-
-protected:
-	/**
-	 *	Pattern to search in source control for the label.
-	 *	This somehow works for older crashes, before 4.2 and for the newest one,
-	 *	bur for the newest we also need to look for the executables on the network drive.
-	 *	This may change in future.
-	 */
-	FString SourceControlBuildLabelPattern;
-	
-	/** Indicates that the crash handler is ready to do work */
-	bool bInitialized;
-
-public:
-	/** A platform independent representation of a crash */
-	FCrashInfo CrashInfo;
-	
 	/** Virtual destructor */
-	virtual ~ICrashDebugHelper()
-	{}
+	virtual ~ICrashDebugHelper() = default;
 
 	/**
 	 *	Initialize the helper
@@ -294,35 +208,12 @@ public:
 		return false;
 	}
 
-	/**
-	 * Sync the branch root relative file names to the requested label
-	 *
-	 *	@param	bOutPDBCacheEntryValid		Returns whether the PDB cache entry was found or created and whether it contains files.
-	 *
-	 *	@return	bool						true if successful, false if not
-	 */
-	virtual bool SyncModules(bool& bOutPDBCacheEntryValid);
-
-	/**
-	 *	Sync a single source file to the requested CL.
-	 */
-	virtual bool SyncSourceFile();
-
 	/**
 	 *	Extract lines from a source file, and add to the crash report.
 	 */
 	virtual void AddSourceToReport();
 
-	/**
-	 *	Extract annotated lines from a source file stored in Perforce, and add to the crash report.
-	 */
-	virtual bool AddAnnotatedSourceToReport();
-
 protected:
-
-	/** Finds the storage of the symbols and the executables for the specified changelist and the depot name, it can be Perforce, network drive or stored locally. */
-	void FindSymbolsAndBinariesStorage();
-
 	/**
 	 *	Load the given ANSI text file to an array of strings - one FString per line of the file.
 	 *	Intended for use in simple text parsing actions
@@ -334,6 +225,10 @@ protected:
 	bool ReadSourceFile( TArray<FString>& OutStrings );
 
 public:
-	bool InitSourceControl(bool bShowLogin);
-	void ShutdownSourceControl();
+	/** A platform independent representation of a crash */
+	FCrashInfo CrashInfo;
+
+protected:
+	/** Indicates that the crash handler is ready to do work */
+	bool bInitialized = false;
 };
diff --git a/Engine/Source/Developer/CrashDebugHelper/Public/CrashDebugPDBCache.h b/Engine/Source/Developer/CrashDebugHelper/Public/CrashDebugPDBCache.h
deleted file mode 100644
diff --git a/Engine/Source/Programs/CrashReportClient/CrashReportClient.Build.cs b/Engine/Source/Programs/CrashReportClient/CrashReportClient.Build.cs
--- a/Engine/Source/Programs/CrashReportClient/CrashReportClient.Build.cs
+++ b/Engine/Source/Programs/CrashReportClient/CrashReportClient.Build.cs
@@ -21,7 +21,6 @@ public class CrashReportClient : ModuleRules
 				"Core",
 				"CoreUObject",
 				"ApplicationCore",
-				"CrashDebugHelper",
 				"CrashReportCore",
 				"HTTP",
 				"Json",
@@ -40,17 +39,20 @@ public class CrashReportClient : ModuleRules
 			}
 		);
 
-		PrivateIncludePathModuleNames.AddRange(
-			new string[] {
+		if (Target.Configuration != UnrealTargetConfiguration.Shipping)
+		{
+			PrivateIncludePathModuleNames.AddRange(
+				new string[] {
 				"SlateReflector",
-			}
-		);
+				}
+			);
 
-		DynamicallyLoadedModuleNames.AddRange(
-			new string[] {
+			DynamicallyLoadedModuleNames.AddRange(
+				new string[] {
 				"SlateReflector",
-			}
-		);
+				}
+			);
+		}
 
 		PrivateIncludePaths.Add("Runtime/Launch/Private");		// For LaunchEngineLoop.cpp include
 
diff --git a/Engine/Source/Programs/CrashReportClient/Private/CrashReportClientApp.cpp b/Engine/Source/Programs/CrashReportClient/Private/CrashReportClientApp.cpp
--- a/Engine/Source/Programs/CrashReportClient/Private/CrashReportClientApp.cpp
+++ b/Engine/Source/Programs/CrashReportClient/Private/CrashReportClientApp.cpp
@@ -23,7 +23,9 @@
 	#include "SCrashReportClient.h"
 	#include "CrashReportClient.h"
 	#include "CrashReportClientStyle.h"
+#if !UE_BUILD_SHIPPING
 	#include "ISlateReflectorModule.h"
+#endif
 	#include "Framework/Application/SlateApplication.h"
 #endif // !CRASH_REPORT_UNATTENDED_ONLY
 
@@ -306,11 +308,13 @@ SubmitCrashReportResult RunWithUI(FPlatformErrorReport ErrorReport)
 	// Setting focus seems to have to happen after the Window has been added
 	FSlateApplication::Get().ClearKeyboardFocus(EFocusCause::Cleared);
 
+#if !UE_BUILD_SHIPPING
 	// Debugging code
 	if (RunWidgetReflector)
 	{
 		FModuleManager::LoadModuleChecked<ISlateReflectorModule>("SlateReflector").DisplayWidgetReflector();
 	}
+#endif
 
 	// Bring the window to the foreground as it may be behind the crashed process
 	Window->HACK_ForceToFront();
diff --git a/Engine/Source/Programs/MinidumpDiagnostics/MinidumpDiagnostics.Build.cs b/Engine/Source/Programs/MinidumpDiagnostics/MinidumpDiagnostics.Build.cs
--- a/Engine/Source/Programs/MinidumpDiagnostics/MinidumpDiagnostics.Build.cs
+++ b/Engine/Source/Programs/MinidumpDiagnostics/MinidumpDiagnostics.Build.cs
@@ -13,6 +13,7 @@ public class MinidumpDiagnostics : ModuleRules
 			new string[] 
 			{ 
 				"Core",
+				"CoreUObject",
 				"ApplicationCore",
 				"CrashDebugHelper",
 				"PerforceSourceControl",
diff --git a/Engine/Source/Programs/MinidumpDiagnostics/MinidumpDiagnostics.Target.cs b/Engine/Source/Programs/MinidumpDiagnostics/MinidumpDiagnostics.Target.cs
--- a/Engine/Source/Programs/MinidumpDiagnostics/MinidumpDiagnostics.Target.cs
+++ b/Engine/Source/Programs/MinidumpDiagnostics/MinidumpDiagnostics.Target.cs
@@ -17,8 +17,7 @@ public class MinidumpDiagnosticsTarget : TargetRules
 
 		// MinidumpDiagnostics doesn't ever compile with the engine linked in
 		bCompileAgainstEngine = false;
-
-		//bCompileICU = false;
+		bCompileAgainstCoreUObject = true;
 
 		// MinidumpDiagnostics.exe has no exports, so no need to verify that a .lib and .exp file was emitted by the linker.
 		bHasExports = false;
diff --git a/Engine/Source/Programs/MinidumpDiagnostics/Private/MinidumpDiagnosticsApp.cpp b/Engine/Source/Programs/MinidumpDiagnostics/Private/MinidumpDiagnosticsApp.cpp
--- a/Engine/Source/Programs/MinidumpDiagnostics/Private/MinidumpDiagnosticsApp.cpp
+++ b/Engine/Source/Programs/MinidumpDiagnostics/Private/MinidumpDiagnosticsApp.cpp
@@ -12,10 +12,8 @@ int32 RunMinidumpDiagnostics(int32 ArgC, TCHAR* Argv[])
 	if( ArgC < 2 )
 	{
 		UE_LOG( LogInit, Error, TEXT( "MinidumpDiagnostics - not enough parameters." ) );
-		UE_LOG( LogInit, Error, TEXT( " ... usage: MinidumpDiagnostics.exe <Crash.dmp> [-Annotate] [-SyncSymbols] [-SyncMicrosoftSymbols]" ) );
+		UE_LOG( LogInit, Error, TEXT( " ... usage: MinidumpDiagnostics.exe <Crash.dmp> [-SyncMicrosoftSymbols]" ) );
 		UE_LOG( LogInit, Error, TEXT( " ..." ) );
-		UE_LOG( LogInit, Error, TEXT( " ... -Annotate: Use Perforce annotation to decorate the source context" ) );
-		UE_LOG( LogInit, Error, TEXT( " ... -SyncSymbols: Sync symbols to the revision specified by the engine version" ) );
 		UE_LOG( LogInit, Error, TEXT( " ... -SyncMicrosoftSymbols: Sync symbols from the Microsoft Symbol Server" ) );
 		return 1;
 	}
diff --git a/Engine/Source/Programs/SymbolDebugger/Private/SymbolDebugger.cpp b/Engine/Source/Programs/SymbolDebugger/Private/SymbolDebugger.cpp
--- a/Engine/Source/Programs/SymbolDebugger/Private/SymbolDebugger.cpp
+++ b/Engine/Source/Programs/SymbolDebugger/Private/SymbolDebugger.cpp
@@ -516,15 +516,6 @@ SSymbolDebugger::ESymbolDebuggerActionResults FSymbolDebugger::ActionHasComplete
 
 void FSymbolDebugger::Tick()
 {
-	static bool s_bFirstTick = true;
-	if (s_bFirstTick)
-	{
-		// Initialize source control and show the login window
-		CrashHelperModule->InitSourceControl(true);
-
-		s_bFirstTick = false;
-	}
-
 	if (CurrentAction == SSymbolDebugger::DebugAction_Inspect)
 	{
 		if (InspectionTask.IsValid() && InspectionTask->IsDone())
diff --git a/Engine/Source/Runtime/CrashReportCore/CrashReportCore.Build.cs b/Engine/Source/Runtime/CrashReportCore/CrashReportCore.Build.cs
--- a/Engine/Source/Runtime/CrashReportCore/CrashReportCore.Build.cs
+++ b/Engine/Source/Runtime/CrashReportCore/CrashReportCore.Build.cs
@@ -32,17 +32,10 @@ public class CrashReportCore : ModuleRules
         {
 			PublicDependencyModuleNames.Add("CrashDebugHelper");
 			PublicDefinitions.Add("WITH_CRASHREPORTER=1");
-
-			// add source control if we are targeting a program
-			if (Target.Type == TargetType.Program)
-			{
-				PublicDependencyModuleNames.Add("SourceControl");
-			}
-       	}
-        else
+		}
+		else
         {
 			PublicDefinitions.Add("WITH_CRASHREPORTER=0");
-            PublicDependencyModuleNames.Add("SourceControl");
-        }
+		}
     }
 }
