Fixed hang (unexpected long time) in CrashReporterClient when processing the minidump to show the user the crash information. In somes cases, it could take several minutes (>15 minutes) leading the user to think that CRC hung or crashed.

The issue seems to appears randomly. It might be a DLL incompatibility, a memory stomp or a race condition, but when it happened, the system tried to search the Windows system modules images on the user disk, scanning all project assets for each module. The modules causing problems were the system ones like the ones from "C:\Windows\System32" or "C:\Program Files\Common Files" and the system was trying to find the images on D: drive, scanning several thousands of files. The issue could be observed with ProcessMonitor and looking at CRC process activity.

The faulty call was 'ReloadWide()'. Commenting it out fixed the issue and did not exhib any obvious regressions. Considering that we configured the MS debug engine to load symbols only when needed, forcing it to flush and reload all the symbols looked misplaced and counter-productive.

Tested CRC with the Editor in a dev environment.
Tested 'debug ensure' followed by 'debug crash'
Tested with a 'real' crash (access violation)
Tested CRC running in 'monitor mode' and in 'standalone mode'.

diff --git a/Engine/Source/Developer/CrashDebugHelper/Private/Windows/WindowsPlatformStackWalkExt.cpp b/Engine/Source/Developer/CrashDebugHelper/Private/Windows/WindowsPlatformStackWalkExt.cpp
--- a/Engine/Source/Developer/CrashDebugHelper/Private/Windows/WindowsPlatformStackWalkExt.cpp
+++ b/Engine/Source/Developer/CrashDebugHelper/Private/Windows/WindowsPlatformStackWalkExt.cpp
@@ -331,12 +331,20 @@ void FWindowsPlatformStackWalkExt::GetModuleInfoDetailed()
 		CrashModule->Patch = HIWORD( VersionInfo.dwProductVersionLS );
 		CrashModule->Revision = LOWORD( VersionInfo.dwProductVersionLS );
 
+		// NOTE: The function below was commented because in some occasions (unknown reason), it triggered a search for modules on the user disk,
+		//       scanning all projects assets for each module, taking an unacceptable long time in very large projects. This could be observed with
+		//       SysInternal ProcessMonitor. The original intent for calling ReloadWide() is unknown, but normally, this function is meant
+		//       to control when the modules are reloaded, for example if the .pdb were about to be deleted, we would want to reload and cache the
+		//       symbols before losing them. Since we configured the debug engine with SYMOPT_DEFERRED_LOADS to only load symbols if needed,
+		//       reloading them all here looks counter-productive. Commenting it did not show any obvious side effects other than fixing the
+		//       very long module search time issue.
 		// Ensure all the images are synced - need the full path here
-		Symbol->ReloadWide( *CrashModule->Name );
+		//Symbol->ReloadWide( *CrashModule->Name );
 	}
 
 	CrashInfo.Modules.Sort( FSortModulesByName() );
 }
+
 bool FWindowsPlatformStackWalkExt::IsOffsetWithinModules( uint64 Offset )
 {
 	for( int ModuleIndex = 0; ModuleIndex < CrashInfo.Modules.Num(); ModuleIndex++ )
