- Added the system image paths to the image pathname because when the issue happens, looking at ETW event with SysInternal process monitor, we can see that CRC is scanning the engine folder (with all assets) to find the system images.
  - This is a specuative fix because the bug doesn't always manifest and it didn't for me when I implemented this fix and didn't anymore for the QA that reported the issue once I proposed this fix.

diff --git a/Engine/Source/Developer/CrashDebugHelper/Private/Windows/WindowsPlatformStackWalkExt.cpp b/Engine/Source/Developer/CrashDebugHelper/Private/Windows/WindowsPlatformStackWalkExt.cpp
--- a/Engine/Source/Developer/CrashDebugHelper/Private/Windows/WindowsPlatformStackWalkExt.cpp
+++ b/Engine/Source/Developer/CrashDebugHelper/Private/Windows/WindowsPlatformStackWalkExt.cpp
@@ -146,7 +146,7 @@ void FWindowsPlatformStackWalkExt::GetExeFileVersionAndModuleList( FCrashModuleI
 		// Get the full path of the module name
 		TCHAR ModuleName[MAX_PATH] = {0};
 		Symbol->GetModuleNameStringWide( DEBUG_MODNAME_IMAGE, ModuleIndex, ModuleBase, ModuleName, MAX_PATH, NULL );
-		
+
 		const FString RelativeModuleName = ExtractRelativePath( TEXT( "binaries" ), ModuleName );
 		// Get the exe, which we extract the version number, so we know what label to sync to
 		if (RelativeModuleName.Len() > 0 && RelativeModuleName.EndsWith( TEXT( ".exe" ) ))
@@ -159,6 +159,10 @@ void FWindowsPlatformStackWalkExt::GetExeFileVersionAndModuleList( FCrashModuleI
 		{
 			CrashInfo.ModuleNames.Add( ModuleName );
 		}
+		else
+		{
+			SystemModuleNames.Add(ModuleName);
+		}
 	}
 
 	// Get the executable version info.
@@ -193,7 +197,7 @@ void FWindowsPlatformStackWalkExt::SetSymbolPathsFromModules()
 
 		// Iterate over all loaded modules.
 		TSet<FString> SymbolPaths;
-		for (const auto& Filename : CrashInfo.ModuleNames)
+		for (const FString& Filename : CrashInfo.ModuleNames)
 		{
 			const FString Path = FPaths::GetPath(Filename);
 			if (Path.Len() > 0)
@@ -202,7 +206,7 @@ void FWindowsPlatformStackWalkExt::SetSymbolPathsFromModules()
 			}
 		}
 
-		for (const auto& SymbolPath : SymbolPaths)
+		for (const FString& SymbolPath : SymbolPaths)
 		{
 			CombinedPath += SymbolPath;
 			CombinedPath += TEXT(";");
@@ -218,8 +222,25 @@ void FWindowsPlatformStackWalkExt::SetSymbolPathsFromModules()
 #endif
 	}
 
+	TSet<FString> ImagePathSet;
+	FString ImagePathnames(CombinedPath);
+	for (const FString& SysImagePathname : SystemModuleNames)
+	{
+		FString Path = FPaths::GetPath(SysImagePathname);
+		if (Path.Len() > 0)
+		{
+			ImagePathSet.Add(Path);
+		}
+	}
+
+	for (const FString& SysImgPath : ImagePathSet)
+	{
+		ImagePathnames += SysImgPath;
+		ImagePathnames += TEXT(";");
+	}
+
 	// Set the symbol path
-	Symbol->SetImagePathWide( *CombinedPath );
+	Symbol->SetImagePathWide( *ImagePathnames );
 	Symbol->SetSymbolPathWide( *CombinedPath );
 
 	// Add in syncing of the Microsoft symbol servers if requested
@@ -243,7 +264,7 @@ void FWindowsPlatformStackWalkExt::SetSymbolPathsFromModules()
 		Symbol->AppendSymbolPathWide( TEXT( "SRV*..\\..\\Intermediate\\SymbolCache*http://msdl.microsoft.com/download/symbols" ) );
 	}
 
-	TCHAR SymbolPath[16384] = { 0 };
+	TCHAR SymbolPath[16384 * 2] = { 0 }; // On large projects, we usually need more than 16K because the list of module is long.
 
 	Symbol->GetSymbolPathWide( SymbolPath, UE_ARRAY_COUNT(SymbolPath), NULL );
 	TArray<FString> SymbolPaths;
diff --git a/Engine/Source/Developer/CrashDebugHelper/Private/Windows/WindowsPlatformStackWalkExt.h b/Engine/Source/Developer/CrashDebugHelper/Private/Windows/WindowsPlatformStackWalkExt.h
--- a/Engine/Source/Developer/CrashDebugHelper/Private/Windows/WindowsPlatformStackWalkExt.h
+++ b/Engine/Source/Developer/CrashDebugHelper/Private/Windows/WindowsPlatformStackWalkExt.h
@@ -60,4 +60,7 @@ struct FWindowsPlatformStackWalkExt
 protected:
 	/** Reference to the crash info. */
 	FCrashInfo& CrashInfo;
+
+	/** List of modules that are not from the engine. */
+	TArray<FString> SystemModuleNames;
 };
