Fix dangling usage of a temp object

diff --git a/Engine/Source/Developer/CrashDebugHelper/Private/Mac/CrashDebugHelperMac.cpp b/Engine/Source/Developer/CrashDebugHelper/Private/Mac/CrashDebugHelperMac.cpp
--- a/Engine/Source/Developer/CrashDebugHelper/Private/Mac/CrashDebugHelperMac.cpp
+++ b/Engine/Source/Developer/CrashDebugHelper/Private/Mac/CrashDebugHelperMac.cpp
@@ -69,13 +69,12 @@ static int32 ParseVersion(TCHAR const* CrashLog, int32& OutMajor, int32& OutMino
 		Found = swscanf(TCHAR_TO_WCHAR(VersionLine), WTEXT("%*s %d.%d.%d (%*d.%*d.%*d-%d+%256ls)"), &OutMajor, &OutMinor, &OutBuild, &OutChangeList, RawBranch);
 		if(Found == 5)
 		{
-			TCHAR* Branch = WCHAR_TO_TCHAR(RawBranch);
-			TCHAR* BranchEnd = FCStringWide::Strchr(Branch, TEXT(')'));
+			TCHAR* BranchEnd = FCStringWide::Strchr(WCHAR_TO_TCHAR(RawBranch), TEXT(')'));
 			if(BranchEnd)
 			{
 				*BranchEnd = TEXT('\0');
 			}
-			OutBranch = Branch;
+			OutBranch = FString(WCHAR_TO_TCHAR(RawBranch));
 		}
 	}
 	return Found;
