More fixes for dangling pointer for Mac

diff --git a/Engine/Source/Developer/CrashDebugHelper/Private/Mac/CrashDebugHelperMac.cpp b/Engine/Source/Developer/CrashDebugHelper/Private/Mac/CrashDebugHelperMac.cpp
--- a/Engine/Source/Developer/CrashDebugHelper/Private/Mac/CrashDebugHelperMac.cpp
+++ b/Engine/Source/Developer/CrashDebugHelper/Private/Mac/CrashDebugHelperMac.cpp
@@ -220,52 +220,52 @@ static int32 ParseExceptionCode(TCHAR const* CrashLog, uint32& OutExceptionCode)
 		}
 		if(Found)
 		{
-			TCHAR* Buffer = WCHAR_TO_TCHAR(RawBuffer);
+			FString Buffer = FString(WCHAR_TO_TCHAR(RawBuffer));
 
-			TCHAR* End = FCStringWide::Strchr(Buffer, TEXT(')'));
+			TCHAR* End = FCStringWide::Strchr(*Buffer, TEXT(')'));
 			if(End)
 			{
 				*End = TEXT('\0');
 			}
-			if(FCStringWide::Strcmp(Buffer, TEXT("SIGQUIT")) == 0)
+			if(FCStringWide::Strcmp(*Buffer, TEXT("SIGQUIT")) == 0)
 			{
 				OutExceptionCode = SIGQUIT;
 			}
-			else if(FCStringWide::Strcmp(Buffer, TEXT("SIGILL")) == 0)
+			else if(FCStringWide::Strcmp(*Buffer, TEXT("SIGILL")) == 0)
 			{
 				OutExceptionCode = SIGILL;
 			}
-			else if(FCStringWide::Strcmp(Buffer, TEXT("SIGEMT")) == 0)
+			else if(FCStringWide::Strcmp(*Buffer, TEXT("SIGEMT")) == 0)
 			{
 				OutExceptionCode = SIGEMT;
 			}
-			else if(FCStringWide::Strcmp(Buffer, TEXT("SIGFPE")) == 0)
+			else if(FCStringWide::Strcmp(*Buffer, TEXT("SIGFPE")) == 0)
 			{
 				OutExceptionCode = SIGFPE;
 			}
-			else if(FCStringWide::Strcmp(Buffer, TEXT("SIGBUS")) == 0)
+			else if(FCStringWide::Strcmp(*Buffer, TEXT("SIGBUS")) == 0)
 			{
 				OutExceptionCode = SIGBUS;
 			}
-			else if(FCStringWide::Strcmp(Buffer, TEXT("SIGSEGV")) == 0)
+			else if(FCStringWide::Strcmp(*Buffer, TEXT("SIGSEGV")) == 0)
 			{
 				OutExceptionCode = SIGSEGV;
 			}
-			else if(FCStringWide::Strcmp(Buffer, TEXT("SIGSYS")) == 0)
+			else if(FCStringWide::Strcmp(*Buffer, TEXT("SIGSYS")) == 0)
 			{
 				OutExceptionCode = SIGSYS;
 			}
-			else if(FCStringWide::Strcmp(Buffer, TEXT("SIGABRT")) == 0)
+			else if(FCStringWide::Strcmp(*Buffer, TEXT("SIGABRT")) == 0)
 			{
 				OutExceptionCode = SIGABRT;
 			}
-			else if(FCStringWide::Strcmp(Buffer, TEXT("SIGTRAP")) == 0)
+			else if(FCStringWide::Strcmp(*Buffer, TEXT("SIGTRAP")) == 0)
 			{
 				OutExceptionCode = SIGTRAP;
 			}
-			else if(FString(Buffer).IsNumeric())
+			else if(Buffer.IsNumeric())
 			{
-				Found = swscanf(TCHAR_TO_WCHAR(Buffer), WTEXT("%u"), &OutExceptionCode);
+				Found = swscanf(TCHAR_TO_WCHAR(*Buffer), WTEXT("%u"), &OutExceptionCode);
 			}
 			else
 			{
@@ -352,22 +352,22 @@ static int32 ParseThreadStackLine(TCHAR const* StackLine, FString& OutModuleName
 		case 3:
 		{
 			int32 Status = -1;
-			const TCHAR* FunctionName = WCHAR_TO_TCHAR(RawFunctionName);
-			ANSICHAR* DemangledName = abi::__cxa_demangle(TCHAR_TO_UTF8(FunctionName), nullptr, nullptr, &Status);
+			FString FunctionName = FString(WCHAR_TO_TCHAR(RawFunctionName));
+			ANSICHAR* DemangledName = abi::__cxa_demangle(TCHAR_TO_UTF8(*FunctionName), nullptr, nullptr, &Status);
 			if(DemangledName && Status == 0)
 			{
 				// C++ function
 				OutFunctionName = FString::Printf(TEXT("%ls "), UTF8_TO_TCHAR(DemangledName));
 			}
-			else if (FCStringWide::Strlen(FunctionName) > 0 && FCStringWide::Strchr(FunctionName, ']'))
+			else if (FCStringWide::Strlen(*FunctionName) > 0 && FCStringWide::Strchr(*FunctionName, ']'))
 			{
 				// ObjC function
-				OutFunctionName = FString::Printf(TEXT("%ls "), FunctionName);
+				OutFunctionName = FString::Printf(TEXT("%ls "), *FunctionName);
 			}
-			else if(FCStringWide::Strlen(FunctionName) > 0)
+			else if(FCStringWide::Strlen(*FunctionName) > 0)
 			{
 				// C function
-				OutFunctionName = FString::Printf(TEXT("%ls() "), FunctionName);
+				OutFunctionName = FString::Printf(TEXT("%ls() "), *FunctionName);
 			}
 		}
 		case 2:
