Change logic to use FString logic vs using raw pointers from a temp buffer

diff --git a/Engine/Source/Developer/CrashDebugHelper/Private/Mac/CrashDebugHelperMac.cpp b/Engine/Source/Developer/CrashDebugHelper/Private/Mac/CrashDebugHelperMac.cpp
--- a/Engine/Source/Developer/CrashDebugHelper/Private/Mac/CrashDebugHelperMac.cpp
+++ b/Engine/Source/Developer/CrashDebugHelper/Private/Mac/CrashDebugHelperMac.cpp
@@ -222,44 +222,44 @@ static int32 ParseExceptionCode(TCHAR const* CrashLog, uint32& OutExceptionCode)
 		{
 			FString Buffer = FString(WCHAR_TO_TCHAR(RawBuffer));
 
-			TCHAR* End = FCStringWide::Strchr(*Buffer, TEXT(')'));
-			if(End)
+			int End = Buffer.Find(TEXT(")"));
+			if(End != INDEX_NONE)
 			{
-				*End = TEXT('\0');
+				Buffer = Buffer.LeftChop(End);
 			}
-			if(FCStringWide::Strcmp(*Buffer, TEXT("SIGQUIT")) == 0)
+			if(Buffer.Find(TEXT("SIGQUIT")) != INDEX_NONE)
 			{
 				OutExceptionCode = SIGQUIT;
 			}
-			else if(FCStringWide::Strcmp(*Buffer, TEXT("SIGILL")) == 0)
+			else if(Buffer.Find(TEXT("SIGILL")) != INDEX_NONE)
 			{
 				OutExceptionCode = SIGILL;
 			}
-			else if(FCStringWide::Strcmp(*Buffer, TEXT("SIGEMT")) == 0)
+			else if(Buffer.Find(TEXT("SIGEMT")) != INDEX_NONE)
 			{
 				OutExceptionCode = SIGEMT;
 			}
-			else if(FCStringWide::Strcmp(*Buffer, TEXT("SIGFPE")) == 0)
+			else if(Buffer.Find(TEXT("SIGFPE")) != INDEX_NONE)
 			{
 				OutExceptionCode = SIGFPE;
 			}
-			else if(FCStringWide::Strcmp(*Buffer, TEXT("SIGBUS")) == 0)
+			else if(Buffer.Find(TEXT("SIGBUS")) != INDEX_NONE)
 			{
 				OutExceptionCode = SIGBUS;
 			}
-			else if(FCStringWide::Strcmp(*Buffer, TEXT("SIGSEGV")) == 0)
+			else if(Buffer.Find(TEXT("SIGSEGV")) != INDEX_NONE)
 			{
 				OutExceptionCode = SIGSEGV;
 			}
-			else if(FCStringWide::Strcmp(*Buffer, TEXT("SIGSYS")) == 0)
+			else if(Buffer.Find(TEXT("SIGSYS")) != INDEX_NONE)
 			{
 				OutExceptionCode = SIGSYS;
 			}
-			else if(FCStringWide::Strcmp(*Buffer, TEXT("SIGABRT")) == 0)
+			else if(Buffer.Find(TEXT("SIGABRT")) != INDEX_NONE)
 			{
 				OutExceptionCode = SIGABRT;
 			}
-			else if(FCStringWide::Strcmp(*Buffer, TEXT("SIGTRAP")) == 0)
+			else if(Buffer.Find(TEXT("SIGTRAP")) != INDEX_NONE)
 			{
 				OutExceptionCode = SIGTRAP;
 			}
