Fixed symbolication in CrashReportClient on Mac with universal applications

diff --git a/Engine/Source/Developer/CrashDebugHelper/Private/Mac/CrashDebugHelperMac.cpp b/Engine/Source/Developer/CrashDebugHelper/Private/Mac/CrashDebugHelperMac.cpp
--- a/Engine/Source/Developer/CrashDebugHelper/Private/Mac/CrashDebugHelperMac.cpp
+++ b/Engine/Source/Developer/CrashDebugHelper/Private/Mac/CrashDebugHelperMac.cpp
@@ -164,6 +164,30 @@ static int32 ParseGraphics(TCHAR const* CrashLog, FString& OutGPUDetails)
 	return bFound;
 }
 
+static int32 ParseArchitecture(TCHAR const* CrashLog, FString& Architecture)
+{
+	bool bFound = false;
+	TCHAR const* Line = FCStringWide::Strstr(CrashLog, TEXT("Code Type:"));
+	int32 Output = 0;
+	while (Line)
+	{
+		Line += FCStringWide::Strlen(TEXT("Code Type:"));
+		TCHAR const* End = FCStringWide::Strchr(Line, TEXT('\r'));
+		if(!End)
+		{
+			End = FCStringWide::Strchr(Line, TEXT('\n'));
+		}
+		check(End);
+		
+		int32 Length = FMath::Min(PATH_MAX, (int32)((uintptr_t)(End - Line)));
+		Architecture.Append(Line, Length);
+		Architecture = Architecture.TrimStartAndEnd();
+		Line = nullptr;
+		bFound = true;
+	}
+	return bFound;
+}
+
 static int32 ParseError(TCHAR const* CrashLog, FString& OutErrorDetails)
 {
 	bool bFound = false;
@@ -384,7 +408,7 @@ static int32 ParseThreadStackLine(TCHAR const* StackLine, FString& OutModuleName
 	return Found;
 }
 
-static int32 SymboliseStackInfo(FPlatformSymbolDatabaseSet& SymbolCache, TArray<FCrashModuleInfo> const& ModuleInfo, FString ModuleName, uint64 const ProgramCounter, FString& OutFunctionName, FString& OutFileName, int32& OutLineNumber)
+static int32 SymboliseStackInfo(FPlatformSymbolDatabaseSet& SymbolCache, TArray<FCrashModuleInfo> const& ModuleInfo, FString ModuleName, FString Architecture, uint64 const ProgramCounter, FString& OutFunctionName, FString& OutFileName, int32& OutLineNumber)
 {
 	FProgramCounterSymbolInfo Info;
 	int32 ValuesSymbolised = 0;
@@ -403,6 +427,7 @@ static int32 SymboliseStackInfo(FPlatformSymbolDatabaseSet& SymbolCache, TArray<
 	if(!Db)
 	{
 		FApplePlatformSymbolDatabase Database;
+		Database.Architecture = Architecture;
 		if(FPlatformSymbolication::LoadSymbolDatabaseForBinary(TEXT(""), Module.Name, Module.Report, Database))
 		{
 			SymbolCache.Add(Database);
@@ -684,7 +709,11 @@ bool FCrashDebugHelperMac::CreateMinidumpDiagnosticReport( const FString& InCras
 			
 			Result = ParseExceptionCode(*CrashDump, CrashInfo.Exception.Code);
 			check(Result == 1);
-			
+
+			FString Architecture;
+			Result = ParseArchitecture(*CrashDump, Architecture);
+			check(Result == 1);
+
 			FCrashThreadInfo ThreadInfo;
 			ThreadInfo.ThreadId = CrashInfo.Exception.ThreadId;
 			ThreadInfo.SuspendCount = 0;
@@ -741,7 +770,7 @@ bool FCrashDebugHelperMac::CreateMinidumpDiagnosticReport( const FString& InCras
 				if(Result > 1 && Result < 4)
 				{
 					// Attempt to resymbolise using CoreSymbolication
-					Result += SymboliseStackInfo(SymbolCache, CrashInfo.Modules, ModuleName, ProgramCounter, FunctionName, FileName, LineNumber);
+					Result += SymboliseStackInfo(SymbolCache, CrashInfo.Modules, ModuleName, Architecture, ProgramCounter, FunctionName, FileName, LineNumber);
 				}
 				
 				// Output in our format based on the fields we actually have
diff --git a/Engine/Source/Runtime/Core/Private/Apple/ApplePlatformSymbolication.cpp b/Engine/Source/Runtime/Core/Private/Apple/ApplePlatformSymbolication.cpp
--- a/Engine/Source/Runtime/Core/Private/Apple/ApplePlatformSymbolication.cpp
+++ b/Engine/Source/Runtime/Core/Private/Apple/ApplePlatformSymbolication.cpp
@@ -243,11 +243,11 @@ bool FApplePlatformSymbolication::LoadSymbolDatabaseForBinary(FString SourceFold
 	if(!bOK && GAllowApplePlatformSymbolication && (IFileManager::Get().FileSize(*BinaryPath) > 0))
 	{
 		CSSymbolicatorRef Symbolicator = OutDatabase.AppleDB;
-		if( CSIsNull(Symbolicator) )
+		if( CSIsNull(Symbolicator) && (OutDatabase.Architecture == TEXT("") || OutDatabase.Architecture == TEXT("X86_64")))
 		{
 			Symbolicator = CSSymbolicatorCreateWithPathAndArchitecture(TCHAR_TO_UTF8(*BinaryPath), CPU_TYPE_X86_64);
 		}
-        if (CSIsNull(Symbolicator))
+        if (CSIsNull(Symbolicator) && (OutDatabase.Architecture == TEXT("") || OutDatabase.Architecture == TEXT("ARM-64")))
         {
             Symbolicator = CSSymbolicatorCreateWithPathAndArchitecture(TCHAR_TO_UTF8(*BinaryPath), CPU_TYPE_ARM64);
         }
diff --git a/Engine/Source/Runtime/Core/Public/Apple/ApplePlatformSymbolication.h b/Engine/Source/Runtime/Core/Public/Apple/ApplePlatformSymbolication.h
--- a/Engine/Source/Runtime/Core/Public/Apple/ApplePlatformSymbolication.h
+++ b/Engine/Source/Runtime/Core/Public/Apple/ApplePlatformSymbolication.h
@@ -32,6 +32,7 @@ struct FApplePlatformSymbolDatabase
 	
 	TSharedPtr<FGenericPlatformSymbolDatabase> GenericDB;
 	FApplePlatformSymbolCache AppleDB;
+	FString Architecture;
 };
 
 /**
