Fix crash report portable callstack on iOS

diff --git a/Engine/Source/Developer/CrashDebugHelper/Private/IOS/CrashDebugHelperIOS.cpp b/Engine/Source/Developer/CrashDebugHelper/Private/IOS/CrashDebugHelperIOS.cpp
--- a/Engine/Source/Developer/CrashDebugHelper/Private/IOS/CrashDebugHelperIOS.cpp
+++ b/Engine/Source/Developer/CrashDebugHelper/Private/IOS/CrashDebugHelperIOS.cpp
@@ -214,58 +214,58 @@ static int32 ParseExceptionCode(TCHAR const* CrashLog, uint32& OutExceptionCode)
 	TCHAR const* Line = FCStringWide::Strstr(CrashLog, TEXT("Exception Type:"));
 	if(Line)
 	{
-		TCHAR Buffer[257] = {0};
-		Found = swscanf(TCHAR_TO_WCHAR(Line), WTEXT("%*s %*s %*s (%256ls)"), Buffer);
+		wchar_t BufferW[257] = {0};
+		Found = swscanf(TCHAR_TO_WCHAR(Line), WTEXT("%*s %*s %*s (%256ls)"), BufferW);
 		if(!Found)
 		{
-			Found = swscanf(TCHAR_TO_WCHAR(Line), WTEXT("%*s %*s %256ls"), Buffer);
+			Found = swscanf(TCHAR_TO_WCHAR(Line), WTEXT("%*s %*s %256ls"), BufferW);
 		}
 		if(Found)
 		{
-			TCHAR* End = FCStringWide::Strchr(Buffer, TEXT(')'));
+			TCHAR* End = FCStringWide::Strchr(WCHAR_TO_TCHAR(BufferW), TEXT(')'));
 			if(End)
 			{
 				*End = TEXT('\0');
 			}
-			if(FCStringWide::Strcmp(Buffer, TEXT("SIGQUIT")) == 0)
+			if(FCStringWide::Strcmp(WCHAR_TO_TCHAR(BufferW), TEXT("SIGQUIT")) == 0)
 			{
 				OutExceptionCode = SIGQUIT;
 			}
-			else if(FCStringWide::Strcmp(Buffer, TEXT("SIGILL")) == 0)
+			else if(FCStringWide::Strcmp(WCHAR_TO_TCHAR(BufferW), TEXT("SIGILL")) == 0)
 			{
 				OutExceptionCode = SIGILL;
 			}
-			else if(FCStringWide::Strcmp(Buffer, TEXT("SIGEMT")) == 0)
+			else if(FCStringWide::Strcmp(WCHAR_TO_TCHAR(BufferW), TEXT("SIGEMT")) == 0)
 			{
 				OutExceptionCode = SIGEMT;
 			}
-			else if(FCStringWide::Strcmp(Buffer, TEXT("SIGFPE")) == 0)
+			else if(FCStringWide::Strcmp(WCHAR_TO_TCHAR(BufferW), TEXT("SIGFPE")) == 0)
 			{
 				OutExceptionCode = SIGFPE;
 			}
-			else if(FCStringWide::Strcmp(Buffer, TEXT("SIGBUS")) == 0)
+			else if(FCStringWide::Strcmp(WCHAR_TO_TCHAR(BufferW), TEXT("SIGBUS")) == 0)
 			{
 				OutExceptionCode = SIGBUS;
 			}
-			else if(FCStringWide::Strcmp(Buffer, TEXT("SIGSEGV")) == 0)
+			else if(FCStringWide::Strcmp(WCHAR_TO_TCHAR(BufferW), TEXT("SIGSEGV")) == 0)
 			{
 				OutExceptionCode = SIGSEGV;
 			}
-			else if(FCStringWide::Strcmp(Buffer, TEXT("SIGSYS")) == 0)
+			else if(FCStringWide::Strcmp(WCHAR_TO_TCHAR(BufferW), TEXT("SIGSYS")) == 0)
 			{
 				OutExceptionCode = SIGSYS;
 			}
-			else if(FCStringWide::Strcmp(Buffer, TEXT("SIGABRT")) == 0)
+			else if(FCStringWide::Strcmp(WCHAR_TO_TCHAR(BufferW), TEXT("SIGABRT")) == 0)
 			{
 				OutExceptionCode = SIGABRT;
 			}
-			else if(FCStringWide::Strcmp(Buffer, TEXT("SIGTRAP")) == 0)
+			else if(FCStringWide::Strcmp(WCHAR_TO_TCHAR(BufferW), TEXT("SIGTRAP")) == 0)
 			{
 				OutExceptionCode = SIGTRAP;
 			}
-			else if(FString(Buffer).IsNumeric())
+			else if(FString(WCHAR_TO_TCHAR(BufferW)).IsNumeric())
 			{
-				Found = swscanf(TCHAR_TO_WCHAR(Buffer), WTEXT("%u"), &OutExceptionCode);
+				Found = swscanf(BufferW, WTEXT("%u"), &OutExceptionCode);
 			}
 			else
 			{
@@ -325,35 +325,35 @@ static TCHAR const* FindCrashedThreadStack(TCHAR const* CrashLog)
 	return Line;
 }
 
-static int32 ParseThreadStackLine(TCHAR const* StackLine, FString& OutModuleName, uint64& OutProgramCounter, FString& OutFunctionName, FString& OutFileName, int32& OutLineNumber)
+static int32 ParseThreadStackLine(TCHAR const* StackLine, FString& OutModuleName, uint64& OutProgramCounter, FString& OutFunctionName, FString& OutFileName, int32& OutLineNumber, uint64& OutFuncAddress, uint64& OutFuncOffset)
 {
-	TCHAR ModuleName[257];
-	TCHAR FunctionName[1025];
-	TCHAR FileName[257];
+	wchar_t ModuleNameW[257];
+	wchar_t FunctionNameW[1025];
+	wchar_t FileNameW[257];
 	
-	int32 Found = swscanf(TCHAR_TO_WCHAR(StackLine), WTEXT("%*d %256ls 0x%lx"), ModuleName, &OutProgramCounter);
+	int32 Found = swscanf(TCHAR_TO_WCHAR(StackLine), WTEXT("%*d %256ls 0x%lx"), ModuleNameW, &OutProgramCounter);
 	if(Found == 2)
 	{
-		uint64 FunctionAddress = 0;
-		uint32 FunctionOffset = 0;
-		if(swscanf(TCHAR_TO_WCHAR(StackLine), WTEXT("%*d %*ls %*lx 0x%lx + %d"), &FunctionAddress, &FunctionOffset) == 0)
+		OutFuncAddress = 0;
+		OutFuncOffset = 0;
+		if(swscanf(TCHAR_TO_WCHAR(StackLine), WTEXT("%*d %*ls 0x%*lx 0x%lx + %d"), &OutFuncAddress, &OutFuncOffset) == 0)
 		{
-			Found += swscanf(TCHAR_TO_WCHAR(StackLine), WTEXT("%*d %*ls %*lx %1024ls + %*d (%256ls:%d)"), FunctionName, FileName, &OutLineNumber);
+			Found += swscanf(TCHAR_TO_WCHAR(StackLine), WTEXT("%*d %*ls 0x%*lx %1024ls + %*d (%256ls:%d)"), FunctionNameW, FileNameW, &OutLineNumber);
 		}
 	}
-	
+		
     switch(Found)
     {
         case 5:
         case 4:
         {
-            OutFileName = FileName;
+            OutFileName = WCHAR_TO_TCHAR(FileNameW);
         }
         case 3:
         {
 #if PLATFORM_IOS
             int32 Status = -1;
-            ANSICHAR* DemangledName = abi::__cxa_demangle(TCHAR_TO_UTF8(FunctionName), nullptr, nullptr, &Status);
+            ANSICHAR* DemangledName = abi::__cxa_demangle(TCHAR_TO_UTF8(WCHAR_TO_TCHAR(FunctionNameW)), nullptr, nullptr, &Status);
             if (DemangledName && Status == 0)
             {
                 // C++ function
@@ -361,21 +361,21 @@ static int32 ParseThreadStackLine(TCHAR const* StackLine, FString& OutModuleName
             }
             else
 #endif
-            if (FCStringWide::Strlen(FunctionName) > 0 && FCStringWide::Strchr(FunctionName, ']'))
+            if (FCStringWide::Strlen(WCHAR_TO_TCHAR(FunctionNameW)) > 0 && FCStringWide::Strchr(WCHAR_TO_TCHAR(FunctionNameW), ']'))
             {
                 // ObjC function
-                OutFunctionName = FString::Printf(TEXT("%ls "), FunctionName);
+                OutFunctionName = FString::Printf(TEXT("%ls "), WCHAR_TO_TCHAR(FunctionNameW));
             }
-            else if (FCStringWide::Strlen(FunctionName) > 0)
+            else if (FCStringWide::Strlen(WCHAR_TO_TCHAR(FunctionNameW)) > 0)
             {
                 // C Function
-                OutFunctionName = FString::Printf(TEXT("%ls() "), FunctionName);
+                OutFunctionName = FString::Printf(TEXT("%ls() "), WCHAR_TO_TCHAR(FunctionNameW));
             }
         }
         case 2:
         case 1:
         {
-            OutModuleName = ModuleName;
+            OutModuleName = WCHAR_TO_TCHAR(ModuleNameW);
         }
         default:
         {
@@ -617,6 +617,8 @@ bool FCrashDebugHelperIOS::CreateMinidumpDiagnosticReport( const FString& InCras
 			FString Gpu;
 			
 			uint64 ProgramCounter = 0;
+			uint64 FunctionAddress = 0;
+			uint64 FunctionOffset = 0;
 			int32 Major = 0;
 			int32 Minor = 0;
 			int32 Build = 0;
@@ -710,7 +712,7 @@ bool FCrashDebugHelperIOS::CreateMinidumpDiagnosticReport( const FString& InCras
 					bIsCrashLocation = (Index == 5);
 				}
 				
-				Result = ParseThreadStackLine(ThreadStackLine, ModuleName, ProgramCounter, FunctionName, FileName, LineNumber);
+				Result = ParseThreadStackLine(ThreadStackLine, ModuleName, ProgramCounter, FunctionName, FileName, LineNumber, FunctionAddress, FunctionOffset);
 				
 				// If we got the modulename & program counter but didn't parse the filename & linenumber we can resymbolise
 				if(Result > 1 && Result < 4)
@@ -720,6 +722,7 @@ bool FCrashDebugHelperIOS::CreateMinidumpDiagnosticReport( const FString& InCras
 				}
 				
 				// Output in our format based on the fields we actually have
+				CrashInfo.Exception.PortableCallStackString.Push( FString::Printf(TEXT("%-*s 0x%016llx + %-16llx"), 40, *ModuleName, FunctionAddress, FunctionOffset));
 				switch (Result)
 				{
 					case 2:
diff --git a/Engine/Source/Developer/CrashDebugHelper/Public/CrashDebugHelper.h b/Engine/Source/Developer/CrashDebugHelper/Public/CrashDebugHelper.h
--- a/Engine/Source/Developer/CrashDebugHelper/Public/CrashDebugHelper.h
+++ b/Engine/Source/Developer/CrashDebugHelper/Public/CrashDebugHelper.h
@@ -64,6 +64,7 @@ public:
 	FString ExceptionString;
 
 	TArray<FString> CallStackString;
+	TArray<FString> PortableCallStackString;
 };
 
 /** 
diff --git a/Engine/Source/Programs/UnrealBuildTool/Platform/IOS/UEBuildIOS.cs b/Engine/Source/Programs/UnrealBuildTool/Platform/IOS/UEBuildIOS.cs
--- a/Engine/Source/Programs/UnrealBuildTool/Platform/IOS/UEBuildIOS.cs
+++ b/Engine/Source/Programs/UnrealBuildTool/Platform/IOS/UEBuildIOS.cs
@@ -82,6 +82,8 @@ namespace UnrealBuildTool
 		/// Don't generate crashlytics data
 		/// </summary>
 		[CommandLine("-alwaysgeneratedsym", Value = "true")]
+		[CommandLine("-EnableDSYM", Value = "true")]
+		[XmlConfigFile(Category = "BuildConfiguration", Name = "bUseDSYMFiles")]
 		public bool bGeneratedSYM = false;
 
 		/// <summary>
diff --git a/Engine/Source/Runtime/CrashReportCore/Private/CrashDescription.cpp b/Engine/Source/Runtime/CrashReportCore/Private/CrashDescription.cpp
--- a/Engine/Source/Runtime/CrashReportCore/Private/CrashDescription.cpp
+++ b/Engine/Source/Runtime/CrashReportCore/Private/CrashDescription.cpp
@@ -126,6 +126,7 @@ FPrimaryCrashProperties::FPrimaryCrashProperties()
 	, GameSessionID( FGenericCrashContext::RuntimePropertiesTag, TEXT( "GameSessionID" ), this )
 	// Multiline properties
 	, CallStack( FGenericCrashContext::RuntimePropertiesTag, TEXT( "CallStack" ), this )
+	, PCallStack( FGenericCrashContext::RuntimePropertiesTag, TEXT( "PCallStack" ), this )
 	, SourceContext( FGenericCrashContext::RuntimePropertiesTag, TEXT( "SourceContext" ), this )
 	, Modules( FGenericCrashContext::RuntimePropertiesTag, TEXT( "Modules" ), this )
 	, UserDescription( FGenericCrashContext::RuntimePropertiesTag, TEXT( "UserDescription" ), this )
diff --git a/Engine/Source/Runtime/CrashReportCore/Private/GenericErrorReport.cpp b/Engine/Source/Runtime/CrashReportCore/Private/GenericErrorReport.cpp
--- a/Engine/Source/Runtime/CrashReportCore/Private/GenericErrorReport.cpp
+++ b/Engine/Source/Runtime/CrashReportCore/Private/GenericErrorReport.cpp
@@ -182,17 +182,20 @@ void FGenericErrorReport::SetPrimaryCrashProperties( FPrimaryCrashProperties& ou
 	if (Helper && bValidCallstack)
 	{
 		TArray<FString> CallStack = Helper->CrashInfo.Exception.CallStackString;
+		TArray<FString> PCallStack = Helper->CrashInfo.Exception.PortableCallStackString;
 
 		// Get the callstack and remove any frames that we don't care about
 		int64 NumMinidumpFramesToIgnore = out_PrimaryCrashProperties.NumMinidumpFramesToIgnore;
 		if (NumMinidumpFramesToIgnore > 0)
 		{
 			CallStack.RemoveAt(0, FMath::Min(CallStack.Num(), (int32)NumMinidumpFramesToIgnore));
+			PCallStack.RemoveAt(0, FMath::Min(PCallStack.Num(), (int32)NumMinidumpFramesToIgnore));
 		}
 
 		if (CallStack.Num() > 0)
 		{
 			out_PrimaryCrashProperties.CallStack = CallStack;
+			out_PrimaryCrashProperties.PCallStack = PCallStack;
 		}
 		out_PrimaryCrashProperties.Modules = Helper->CrashInfo.ModuleNames;
 		out_PrimaryCrashProperties.SourceContext = Helper->CrashInfo.SourceContext;
diff --git a/Engine/Source/Runtime/CrashReportCore/Public/CrashDescription.h b/Engine/Source/Runtime/CrashReportCore/Public/CrashDescription.h
--- a/Engine/Source/Runtime/CrashReportCore/Public/CrashDescription.h
+++ b/Engine/Source/Runtime/CrashReportCore/Public/CrashDescription.h
@@ -316,6 +316,13 @@ struct FPrimaryCrashProperties
 	 * 
 	 */
 	FCrashProperty CallStack;
+	
+	/**
+	 * An array of FStrings representing the portable callstack of the crash.
+	 * @RawCallStack	varchar(MAX)
+	 *
+	 */
+	FCrashProperty PCallStack;
 
 	/**
 	 * An array of FStrings showing the source code around the crash.
