Add background thread info to iOS crash report

#rnx

diff --git a/Engine/Source/Developer/CrashDebugHelper/Private/IOS/CrashDebugHelperIOS.cpp b/Engine/Source/Developer/CrashDebugHelper/Private/IOS/CrashDebugHelperIOS.cpp
--- a/Engine/Source/Developer/CrashDebugHelper/Private/IOS/CrashDebugHelperIOS.cpp
+++ b/Engine/Source/Developer/CrashDebugHelper/Private/IOS/CrashDebugHelperIOS.cpp
@@ -662,10 +662,6 @@ bool FCrashDebugHelperIOS::CreateMinidumpDiagnosticReport( const FString& InCras
 			Result = ParseExceptionCode(*CrashDump, CrashInfo.Exception.Code);
 			check(Result == 1);
 			
-			FCrashThreadInfo ThreadInfo;
-			ThreadInfo.ThreadId = CrashInfo.Exception.ThreadId;
-			ThreadInfo.SuspendCount = 0;
-			
 			// Parse modules now for symbolication - if we don't have the running process we need to symbolicate by UUID
 			TCHAR const* ModuleLine = FindModules(*CrashDump);
 			while(ModuleLine)
@@ -689,94 +685,130 @@ bool FCrashDebugHelperIOS::CreateMinidumpDiagnosticReport( const FString& InCras
 			
 			FPlatformSymbolDatabaseSet SymbolCache;
 			
-			bool bIsCrashLocation = true;
-			TCHAR const* ThreadStackLine = FindCrashedThreadStack(*CrashDump);
-			uint32 Index = 0;
-			while(ThreadStackLine)
+			
+			uint32 CrashedThreadNumber = 0;
+			ParseCrashedThread(*CrashDump, CrashedThreadNumber);
+			
+			for (int ThreadNumber = 0; true; ++ThreadNumber)
 			{
-				if(CrashInfo.Exception.Code == SIGTRAP)
+				TCHAR const* ThreadStackLine = FindThreadStack(*CrashDump, ThreadNumber);
+				if (!ThreadStackLine)
 				{
-					// For ensures strip the first three lines as they are PLCrashReporter nonsense
-					if(Index < 3)
-					{
-						ThreadStackLine = FCStringWide::Strchr(ThreadStackLine, TEXT('\n'));
-						if(ThreadStackLine)
-						{
-							ThreadStackLine += 1;
-						}
-						++Index;
-						continue;
-					}
-					
-					// Crash location is the 5th entry in the stack.
-					bIsCrashLocation = (Index == 5);
+					break;
 				}
 				
-				Result = ParseThreadStackLine(ThreadStackLine, ModuleName, ProgramCounter, FunctionName, FileName, LineNumber, FunctionAddress, FunctionOffset);
+				FCrashThreadInfo ThreadInfo;
+				ThreadInfo.SuspendCount = 0;
+				ThreadInfo.IsCrashing = (CrashedThreadNumber == ThreadNumber);
+				ThreadInfo.ThreadId = ThreadNumber; // PLReporter does not preserve ThreadId, just use index number
+				ThreadInfo.ThreadName = FString::Printf(TEXT("Thread %d"), ThreadInfo.ThreadId); // PLReporter does not preserve ThreadName either
 				
-				// If we got the modulename & program counter but didn't parse the filename & linenumber we can resymbolise
-				if(Result > 1 && Result < 4)
+				bool bIsCrashLocation = true;
+				uint32 Index = 0;
+				while(ThreadStackLine)
 				{
-					// Attempt to resymbolise using CoreSymbolication
-					Result += SymboliseStackInfo(SymbolCache, CrashInfo.Modules, ModuleName, ProgramCounter, FunctionName, FileName, LineNumber);
-				}
-				
-				// Output in our format based on the fields we actually have
-				CrashInfo.Exception.PortableCallStackString.Push( FString::Printf(TEXT("%-*s 0x%016llx + %-16llx"), 40, *ModuleName, FunctionAddress, FunctionOffset));
-				switch (Result)
-				{
-					case 2:
-						CrashInfo.Exception.CallStackString.Push( FString::Printf( TEXT( "Unknown() Address = 0x%lx (filename not found) [in %s]" ), ProgramCounter, *ModuleName ) );
-						ThreadInfo.CallStack.Push(ProgramCounter);
-						
-						ThreadStackLine = FCStringWide::Strchr(ThreadStackLine, TEXT('\n'));
-						check(ThreadStackLine);
-						ThreadStackLine += 1;
-						break;
-						
-					case 3:
-					case 4:
-						CrashInfo.Exception.CallStackString.Push( FString::Printf( TEXT( "%s Address = 0x%lx (filename not found) [in %s]" ), *FunctionName, ProgramCounter, *ModuleName ) );
-						ThreadInfo.CallStack.Push(ProgramCounter);
-						
-						ThreadStackLine = FCStringWide::Strchr(ThreadStackLine, TEXT('\n'));
-						check(ThreadStackLine);
-						ThreadStackLine += 1;
-						break;
-						
-					case 5:
-					case 6: // Function name might be parsed twice
-						if(bIsCrashLocation)
+					if(ThreadInfo.IsCrashing && CrashInfo.Exception.Code == SIGTRAP)
+					{
+						// For ensures strip the first three lines as they are PLCrashReporter nonsense
+						if(Index < 3)
 						{
-							if( FileName.Len() > 0 && LineNumber > 0 )
+							ThreadStackLine = FCStringWide::Strchr(ThreadStackLine, TEXT('\n'));
+							if(ThreadStackLine)
 							{
-								// Sync the source file where the crash occurred
-								CrashInfo.SourceFile = ExtractRelativePath( TEXT( "source" ), *FileName );
-								CrashInfo.SourceLineNumber = LineNumber;
-								
-								// Add the standard source context
-								AddSourceToReport();
+								ThreadStackLine += 1;
 							}
+							++Index;
+							continue;
 						}
 						
-						CrashInfo.Exception.CallStackString.Push( FString::Printf( TEXT( "%s Address = 0x%lx [%s, line %d] [in %s]" ), *FunctionName, ProgramCounter, *FileName, LineNumber, *ModuleName ) );
-						ThreadInfo.CallStack.Push(ProgramCounter);
-						
-						ThreadStackLine = FCStringWide::Strchr(ThreadStackLine, TEXT('\n'));
-						check(ThreadStackLine);
-						ThreadStackLine += 1;
-						break;
-						
-					default:
-						ThreadStackLine = nullptr;
-						break;
+						// Crash location is the 5th entry in the stack.
+						bIsCrashLocation = (Index == 5);
+					}
+					
+					Result = ParseThreadStackLine(ThreadStackLine, ModuleName, ProgramCounter, FunctionName, FileName, LineNumber, FunctionAddress, FunctionOffset);
+					
+					// If we got the modulename & program counter but didn't parse the filename & linenumber we can resymbolise
+					if(Result > 1 && Result < 4)
+					{
+						// Attempt to resymbolise using CoreSymbolication
+						Result += SymboliseStackInfo(SymbolCache, CrashInfo.Modules, ModuleName, ProgramCounter, FunctionName, FileName, LineNumber);
+					}
+					
+					if (ThreadInfo.IsCrashing)
+					{
+						// Output in our format based on the fields we actually have
+						CrashInfo.Exception.PortableCallStackString.Push( FString::Printf(TEXT("%-*s 0x%016llx + %-16llx"), 40, *ModuleName, FunctionAddress, FunctionOffset));
+					}
+					ThreadInfo.PortableCallStackString += FString::Printf(TEXT("%-*s 0x%016llx + %-16llx\n"), 40, *ModuleName, FunctionAddress, FunctionOffset);
+					
+					switch (Result)
+					{
+						case 2:
+							if (ThreadInfo.IsCrashing)
+							{
+								CrashInfo.Exception.CallStackString.Push( FString::Printf( TEXT( "Unknown() Address = 0x%lx (filename not found) [in %s]" ), ProgramCounter, *ModuleName ) );
+							}
+							ThreadInfo.CallStack.Push(ProgramCounter);
+							
+							ThreadStackLine = FCStringWide::Strchr(ThreadStackLine, TEXT('\n'));
+							check(ThreadStackLine);
+							ThreadStackLine += 1;
+							break;
+							
+						case 3:
+						case 4:
+							if (ThreadInfo.IsCrashing)
+							{
+								CrashInfo.Exception.CallStackString.Push( FString::Printf( TEXT( "%s Address = 0x%lx (filename not found) [in %s]" ), *FunctionName, ProgramCounter, *ModuleName ) );
+							}
+							ThreadInfo.CallStack.Push(ProgramCounter);
+							
+							ThreadStackLine = FCStringWide::Strchr(ThreadStackLine, TEXT('\n'));
+							check(ThreadStackLine);
+							ThreadStackLine += 1;
+							break;
+							
+						case 5:
+						case 6: // Function name might be parsed twice
+							if(ThreadInfo.IsCrashing && bIsCrashLocation)
+							{
+								if( FileName.Len() > 0 && LineNumber > 0 )
+								{
+									// Sync the source file where the crash occurred
+									CrashInfo.SourceFile = ExtractRelativePath( TEXT( "source" ), *FileName );
+									CrashInfo.SourceLineNumber = LineNumber;
+									
+									// Add the standard source context
+									AddSourceToReport();
+								}
+							}
+							
+							if (ThreadInfo.IsCrashing)
+							{
+								CrashInfo.Exception.CallStackString.Push( FString::Printf( TEXT( "%s Address = 0x%lx [%s, line %d] [in %s]" ), *FunctionName, ProgramCounter, *FileName, LineNumber, *ModuleName ) );
+							}
+							ThreadInfo.CallStack.Push(ProgramCounter);
+							
+							ThreadStackLine = FCStringWide::Strchr(ThreadStackLine, TEXT('\n'));
+							check(ThreadStackLine);
+							ThreadStackLine += 1;
+							break;
+							
+						default:
+							ThreadStackLine = nullptr;
+							break;
+					}
+					
+					++Index;
+					bIsCrashLocation = false;
 				}
 				
-				++Index;
-				bIsCrashLocation = false;
+				CrashInfo.Threads.Push(ThreadInfo);
 			}
 			
-			CrashInfo.Threads.Push(ThreadInfo);
+			TCHAR const* ThreadStackLine = FindCrashedThreadStack(*CrashDump);
+			
+			
 			
 			bOK = true;
 		}
diff --git a/Engine/Source/Developer/CrashDebugHelper/Public/CrashDebugHelper.h b/Engine/Source/Developer/CrashDebugHelper/Public/CrashDebugHelper.h
--- a/Engine/Source/Developer/CrashDebugHelper/Public/CrashDebugHelper.h
+++ b/Engine/Source/Developer/CrashDebugHelper/Public/CrashDebugHelper.h
@@ -43,11 +43,14 @@ class FCrashThreadInfo
 {
 public:
 	FString Report;
+	FString ThreadName;
 
 	uint32 ThreadId = 0;
 	uint32 SuspendCount = 0;
+	bool IsCrashing = false;
 
 	TArray<uint64> CallStack;
+	FString PortableCallStackString;
 };
 
 /** 
diff --git a/Engine/Source/Runtime/CrashReportCore/Private/CrashDescription.cpp b/Engine/Source/Runtime/CrashReportCore/Private/CrashDescription.cpp
--- a/Engine/Source/Runtime/CrashReportCore/Private/CrashDescription.cpp
+++ b/Engine/Source/Runtime/CrashReportCore/Private/CrashDescription.cpp
@@ -109,6 +109,44 @@ FCrashPropertyXmlNode& FCrashPropertyXmlNode::operator=(const FXmlNode* Node)
 	return *this;
 }
 
+FCrashPropertyThreadsNode::FCrashPropertyThreadsNode(const FString& InMainCategory, const FString& InSecondCategory, FPrimaryCrashProperties* InOwner)
+: Owner(InOwner)
+	, MainCategory(InMainCategory)
+	, SecondCategory(InSecondCategory)
+	, bSet(false)
+{ }
+
+
+#if WITH_CRASHREPORTER
+FCrashPropertyThreadsNode& FCrashPropertyThreadsNode::operator=(const TArray<FCrashThreadInfo>& ThreadsInfo)
+{
+	if (ThreadsInfo.Num() <= 0)
+	{
+		return *this;
+	}
+	
+	bSet = true;
+	FXmlNode* RuntimePropertiesNode = Owner->XmlFile->GetRootNode()->FindChildNode( FGenericCrashContext::RuntimePropertiesTag );
+	RuntimePropertiesNode->AppendChildNode(TEXT("Threads"));
+	FXmlNode* ThreadsNode = RuntimePropertiesNode->GetChildrenNodes().Last();
+	for (int32 Idx = 0; Idx < ThreadsInfo.Num(); ++Idx)
+	{
+		ThreadsNode->AppendChildNode(TEXT("Thread"));
+		FXmlNode* ThreadNode = ThreadsNode->GetChildrenNodes().Last();
+		
+		const FCrashThreadInfo& ThreadInfo = ThreadsInfo[Idx];
+		ThreadNode->AppendChildNode(TEXT("CallStack"), ThreadInfo.PortableCallStackString);
+		ThreadNode->AppendChildNode(TEXT("IsCrashed"), ThreadInfo.IsCrashing ? TEXT("true") : TEXT("false"));
+		// TODO: do we need thread register states?
+		ThreadNode->AppendChildNode(TEXT("Registers"));
+		ThreadNode->AppendChildNode(TEXT("ThreadID"), FString::Printf(TEXT("%d"), ThreadInfo.ThreadId));
+		ThreadNode->AppendChildNode(TEXT("ThreadName"), ThreadInfo.ThreadName);
+	}
+	
+	return *this;
+}
+#endif
+
 /*-----------------------------------------------------------------------------
 	FPrimaryCrashProperties
 -----------------------------------------------------------------------------*/
diff --git a/Engine/Source/Runtime/CrashReportCore/Private/GenericErrorReport.cpp b/Engine/Source/Runtime/CrashReportCore/Private/GenericErrorReport.cpp
--- a/Engine/Source/Runtime/CrashReportCore/Private/GenericErrorReport.cpp
+++ b/Engine/Source/Runtime/CrashReportCore/Private/GenericErrorReport.cpp
@@ -205,6 +205,8 @@ void FGenericErrorReport::SetPrimaryCrashProperties( FPrimaryCrashProperties& ou
 			FSHA1::HashBuffer(*StringToHash, StringToHash.Len() * sizeof(FString::ElementType), PCallStackHash.Hash);
 			out_PrimaryCrashProperties.PCallStackHashProperty = PCallStackHash.ToString();
 		}
+		
+		out_PrimaryCrashProperties.Threads = Helper->CrashInfo.Threads;
 		out_PrimaryCrashProperties.Modules = Helper->CrashInfo.ModuleNames;
 		out_PrimaryCrashProperties.SourceContext = Helper->CrashInfo.SourceContext;
 
diff --git a/Engine/Source/Runtime/CrashReportCore/Public/CrashDescription.h b/Engine/Source/Runtime/CrashReportCore/Public/CrashDescription.h
--- a/Engine/Source/Runtime/CrashReportCore/Public/CrashDescription.h
+++ b/Engine/Source/Runtime/CrashReportCore/Public/CrashDescription.h
@@ -13,6 +13,10 @@
 #include "XmlFile.h"
 #include "XmlNode.h"
 
+#if WITH_CRASHREPORTER
+#include "CrashDebugHelper.h"
+#endif
+
 enum class ECrashDescVersions : int32;
 enum class ECrashDumpMode : int32;
 class FXmlNode;
@@ -170,11 +174,38 @@ protected:
 	mutable bool bSet;
 };
 
+struct FCrashPropertyThreadsNode
+{
+	friend struct FPrimaryCrashProperties;
+	
+#if WITH_CRASHREPORTER
+	/** Assignment operator for . */
+	FCrashPropertyThreadsNode& operator=(const TArray<FCrashThreadInfo>& ThreadsInfo);
+#endif
+	
+protected:
+	/** Initialization constructor. */
+	FCrashPropertyThreadsNode(const FString& InMainCategory, const FString& InSecondCategory, FPrimaryCrashProperties* InOwner);
+
+protected:
+	/** Owner of the property. */
+	FPrimaryCrashProperties* Owner;
+
+	/** Main category in the crash context. */
+	FString MainCategory;
+
+	/** Second category in the crash context. */
+	FString SecondCategory;
+
+	mutable bool bSet;
+};
+
 /** Primary crash properties required by the crash report system. */
 struct FPrimaryCrashProperties
 {
 	friend struct FCrashProperty;
 	friend struct FCrashPropertyXmlNode;
+	friend struct FCrashPropertyThreadsNode;
 
 	/** Version. */
 	ECrashDescVersions CrashVersion;
@@ -436,7 +467,7 @@ struct FPrimaryCrashProperties
 	/**
 	 *	Thread contexts, XML elements containing info specific to an active thread, e.g. callstacks
 	 */
-	FCrashPropertyXmlNode Threads;
+	FCrashPropertyThreadsNode Threads;
 
 	/**
 	 *	Optional additional data for platform properties
