fixup dbghelp.dll for Windows Arm64

- using version: 20230202.1638.0
- https://www.nuget.org/packages/Microsoft.Debugging.Platform.DbgEng/20230202.1638.0#readme-body-tab

diff --git a/Engine/Source/Developer/CrashDebugHelper/Private/Windows/WindowsPlatformStackWalkExt.cpp b/Engine/Source/Developer/CrashDebugHelper/Private/Windows/WindowsPlatformStackWalkExt.cpp
--- a/Engine/Source/Developer/CrashDebugHelper/Private/Windows/WindowsPlatformStackWalkExt.cpp
+++ b/Engine/Source/Developer/CrashDebugHelper/Private/Windows/WindowsPlatformStackWalkExt.cpp
@@ -246,6 +246,9 @@ void FWindowsPlatformStackWalkExt::SetSymbolPathsFromModules()
 	if( FParse::Param( FCommandLine::Get(), TEXT( "SyncMicrosoftSymbols" ) ) )
 	{
 		FString BinariesDir = FString(FPlatformProcess::BaseDir());
+#if PLATFORM_CPU_ARM_FAMILY && !PLATFORM_WINDOWS_ARM64EC
+		BinariesDir = FPaths::Combine(BinariesDir, TEXT("arm64"));
+#endif
 		if ( !FPaths::FileExists( FPaths::Combine(*BinariesDir, TEXT("symsrv.dll")) ) )
 		{
 			UE_LOG( LogCrashDebugHelper, Warning, TEXT( "Error: symsrv.dll was not detected in: %s. Microsoft symbols will not be downloaded!" ), *BinariesDir );
diff --git a/Engine/Source/Runtime/Core/Core.Build.cs b/Engine/Source/Runtime/Core/Core.Build.cs
--- a/Engine/Source/Runtime/Core/Core.Build.cs
+++ b/Engine/Source/Runtime/Core/Core.Build.cs
@@ -93,9 +93,10 @@ public class Core : ModuleRules
 			
 			if (Target.WindowsPlatform.bUseBundledDbgHelp)
 			{
+				string ArchDir = (Target.Architecture == UnrealArch.Arm64) ? "arm64/" : "";
 				PublicDelayLoadDLLs.Add("DBGHELP.DLL");
 				PrivateDefinitions.Add("USE_BUNDLED_DBGHELP=1");
-				RuntimeDependencies.Add("$(EngineDir)/Binaries/ThirdParty/DbgHelp/dbghelp.dll");
+				RuntimeDependencies.Add($"$(EngineDir)/Binaries/ThirdParty/DbgHelp/{ArchDir}dbghelp.dll");
 			}
 			else
 			{
diff --git a/Engine/Source/Runtime/Core/Private/Windows/WindowsPlatformMisc.cpp b/Engine/Source/Runtime/Core/Private/Windows/WindowsPlatformMisc.cpp
--- a/Engine/Source/Runtime/Core/Private/Windows/WindowsPlatformMisc.cpp
+++ b/Engine/Source/Runtime/Core/Private/Windows/WindowsPlatformMisc.cpp
@@ -903,7 +903,11 @@ void FWindowsPlatformMisc::PlatformPreInit()
 	if (IsWindows8OrGreater())
 	{
 		// Try to load a bundled copy of dbghelp.dll. A bug with Windows 10 version 1709 
+#if PLATFORM_CPU_ARM_FAMILY && !PLATFORM_WINDOWS_ARM64EC
+		FString DbgHlpPath = FPaths::EngineDir() / TEXT("Binaries/ThirdParty/DbgHelp/arm64/dbghelp.dll");
+#else
 		FString DbgHlpPath = FPaths::EngineDir() / TEXT("Binaries/ThirdParty/DbgHelp/dbghelp.dll");
+#endif
 		FPlatformProcess::GetDllHandle(*DbgHlpPath);
 	}
 #endif
